- Split off mesh description into its own class and engine module. The intention is this will be shared by other mesh objects as the source mesh format, rather than being part of UEditableMesh.

diff --git a/Engine/Plugins/Runtime/EditableMesh/Source/EditableMesh/Public/EditableMeshTypes.h b/Engine/Plugins/Runtime/EditableMesh/Source/EditableMesh/Public/EditableMeshTypes.h
--- a/Engine/Plugins/Runtime/EditableMesh/Source/EditableMesh/Public/EditableMeshTypes.h
+++ b/Engine/Plugins/Runtime/EditableMesh/Source/EditableMesh/Public/EditableMeshTypes.h
@@ -5,6 +5,7 @@
 #include "CoreMinimal.h"
 #include "UObject/Object.h"
 #include "UObject/ScriptMacros.h"
+#include "MeshTypes.h"
 #include "EditableMeshTypes.generated.h"
 
 // @todo mesheditor: Move elsewhere
@@ -109,204 +110,6 @@ enum class EEditableMeshElementType
 };
 
 
-// @todo mesheditor: Need comments
-
-USTRUCT( BlueprintType )
-struct FElementID	// @todo mesheditor script: BP doesn't have name spaces, so we might need a more specific display name, or just rename our various types
-{
-	GENERATED_BODY()
-
-	FElementID()
-	{
-	}
-
-	explicit FElementID( const int32 InitIDValue )
-		: IDValue( InitIDValue )
-	{
-	}
-
-	FORCEINLINE int32 GetValue() const
-	{
-		return IDValue;
-	}
-
-	FORCEINLINE bool operator==( const FElementID& Other ) const
-	{
-		return IDValue == Other.IDValue;
-	}
-
-	FORCEINLINE bool operator!=( const FElementID& Other ) const
-	{
-		return IDValue != Other.IDValue;
-	}
-
-	FString ToString() const
-	{
-		return ( IDValue == Invalid.GetValue() ) ? TEXT( "Invalid" ) : FString::Printf( TEXT( "%d" ), IDValue );
-	}
-
-	friend FArchive& operator<<( FArchive& Ar, FElementID& Element )
-	{
-		Ar << Element.IDValue;
-		return Ar;
-	}
-
-	/** Invalid element ID */
-	EDITABLEMESH_API static const FElementID Invalid;
-
-protected:
-
-	/** The actual mesh element index this ID represents.  Read-only. */
-	UPROPERTY( BlueprintReadOnly, Category="Editable Mesh" )
-	int32 IDValue;
-};
-
-
-USTRUCT( BlueprintType )
-struct FVertexID : public FElementID
-{
-	GENERATED_BODY()
-
-	FVertexID()
-	{
-	}
-
-	explicit FVertexID( const FElementID InitElementID )
-		: FElementID( InitElementID.GetValue() )
-	{
-	}
-
-	explicit FVertexID( const int32 InitIDValue )
-		: FElementID( InitIDValue )
-	{
-	}
-
-	FORCEINLINE friend uint32 GetTypeHash( const FVertexID& Other )
-	{
-		return GetTypeHash( Other.IDValue );
-	}
-
-	/** Invalid vertex ID */
-	EDITABLEMESH_API static const FVertexID Invalid;
-};
-
-
-USTRUCT( BlueprintType )
-struct FVertexInstanceID : public FElementID
-{
-	GENERATED_BODY()
-
-	FVertexInstanceID()
-	{
-	}
-
-	explicit FVertexInstanceID( const FElementID InitElementID )
-		: FElementID( InitElementID.GetValue() )
-	{
-	}
-
-	explicit FVertexInstanceID( const uint32 InitIDValue )
-		: FElementID( InitIDValue )
-	{
-	}
-
-	FORCEINLINE friend uint32 GetTypeHash( const FVertexInstanceID& Other )
-	{
-		return GetTypeHash( Other.IDValue );
-	}
-
-	/** Invalid rendering vertex ID */
-	EDITABLEMESH_API static const FVertexInstanceID Invalid;
-};
-
-
-USTRUCT( BlueprintType )
-struct FEdgeID : public FElementID
-{
-	GENERATED_BODY()
-
-	FEdgeID()
-	{
-	}
-
-	explicit FEdgeID( const FElementID InitElementID )
-		: FElementID( InitElementID.GetValue() )
-	{
-	}
-
-	explicit FEdgeID( const int32 InitIDValue )
-		: FElementID( InitIDValue )
-	{
-	}
-
-	FORCEINLINE friend uint32 GetTypeHash( const FEdgeID& Other )
-	{
-		return GetTypeHash( Other.IDValue );
-	}
-
-	/** Invalid edge ID */
-	EDITABLEMESH_API static const FEdgeID Invalid;
-};
-
-
-USTRUCT( BlueprintType )
-struct FPolygonGroupID : public FElementID
-{
-	GENERATED_BODY()
-
-	FPolygonGroupID()
-	{
-	}
-
-	explicit FPolygonGroupID( const FElementID InitElementID )
-		: FElementID( InitElementID.GetValue() )
-	{
-	}
-
-	explicit FPolygonGroupID( const int32 InitIDValue )
-		: FElementID( InitIDValue )
-	{
-	}
-
-	FORCEINLINE friend uint32 GetTypeHash( const FPolygonGroupID& Other )
-	{
-		return GetTypeHash( Other.IDValue );
-	}
-
-	/** Invalid section ID */
-	EDITABLEMESH_API static const FPolygonGroupID Invalid;
-};
-
-
-USTRUCT( BlueprintType )
-struct FPolygonID : public FElementID
-{
-	GENERATED_BODY()
-
-	FPolygonID()
-	{
-	}
-
-	explicit FPolygonID( const FElementID InitElementID )
-		: FElementID( InitElementID.GetValue() )
-	{
-	}
-
-	explicit FPolygonID( const int32 InitIDValue )
-		: FElementID( InitIDValue )
-	{
-	}
-
-	FORCEINLINE friend uint32 GetTypeHash( const FPolygonID& Other )
-	{
-		return GetTypeHash( Other.IDValue );
-	}
-
-	/** Invalid polygon ID */
-	EDITABLEMESH_API static const FPolygonID Invalid;	// @todo mesheditor script: Can we expose these to BP nicely?	Do we even need to?
-};
-
-
 /**
  * Unique identifies a specific sub-mesh within a component
  */
@@ -365,111 +168,6 @@ struct FEditableMeshSubMeshAddress
 };
 
 
-UCLASS( abstract )
-class EDITABLEMESH_API UEditableMeshAttribute : public UObject
-{
-	GENERATED_BODY()
-
-public:
-
-	//
-	// Vertex data for any vertex
-	//
-
-	/** Static: The attribute name for vertex position */
-	UFUNCTION( BlueprintPure, Category="Editable Mesh" ) static inline FName VertexPosition()
-	{
-		return VertexPositionName;
-	}
-
-	/** Static: The attribute name for vertex corner sharpness (only applies to subdivision meshes) */
-	UFUNCTION( BlueprintPure, Category="Editable Mesh" ) static inline FName VertexCornerSharpness()
-	{
-		return VertexCornerSharpnessName;
-	}
-
-	//
-	// Vertex instance data
-	//
-
-	/** Static: The attribute name for vertex normal (tangent Z) */
-	UFUNCTION( BlueprintPure, Category="Editable Mesh" ) static inline FName VertexNormal()
-	{
-		return VertexNormalName;
-	}
-
-	/** Static: The attribute name for vertex tangent vector (tangent X) */
-	UFUNCTION( BlueprintPure, Category="Editable Mesh" ) static inline FName VertexTangent()
-	{
-		return VertexTangentName;
-	}
-
-	/** Static: The attribute name for the vertex basis determinant sign (used to calculate the direction of tangent Y) */
-	UFUNCTION( BlueprintPure, Category="Editable Mesh" ) static inline FName VertexBinormalSign()
-	{
-		return VertexBinormalSignName;
-	}
-
-	/** Static: The attribute name for vertex texture coordinate.  The attribute index defines which texture coordinate set. */
-	UFUNCTION( BlueprintPure, Category="Editable Mesh" ) static inline FName VertexTextureCoordinate()
-	{
-		return VertexTextureCoordinateName;
-	}
-
-	/** Static: The attribute name for the vertex color. */
-	UFUNCTION( BlueprintPure, Category="Editable Mesh" ) static inline FName VertexColor()
-	{
-		return VertexColorName;
-	}
-
-	//
-	// Edges
-	//
-
-	/** Static: The attribute name for edge hardedness */
-	UFUNCTION( BlueprintPure, Category="Editable Mesh" ) static inline FName EdgeIsHard()
-	{
-		return EdgeIsHardName;
-	}
-
-	/** Static: The attribute name for edge crease sharpness (only applies to subdivision meshes) */
-	UFUNCTION( BlueprintPure, Category="Editable Mesh" ) static inline FName EdgeCreaseSharpness()
-	{
-		return EdgeCreaseSharpnessName;
-	}
-
-	//
-	// Polygons
-	//
-
-	/** Static: The attribute name for polygon normal */
-	UFUNCTION( BlueprintPure, Category="Editable Mesh" ) static inline FName PolygonNormal()
-	{
-		return PolygonNormalName;
-	}
-
-	/** Static: The attribute name for polygon center */
-	UFUNCTION( BlueprintPure, Category="Editable Mesh" ) static inline FName PolygonCenter()
-	{
-		return PolygonCenterName;
-	}
-
-private:
-
-	static const FName VertexPositionName;
-	static const FName VertexCornerSharpnessName;
-	static const FName VertexNormalName;
-	static const FName VertexTangentName;
-	static const FName VertexBinormalSignName;
-	static const FName VertexTextureCoordinateName;
-	static const FName VertexColorName;
-	static const FName EdgeIsHardName;
-	static const FName EdgeCreaseSharpnessName;
-	static const FName PolygonNormalName;
-	static const FName PolygonCenterName;
-};
-
-
 UENUM()
 enum class EMeshModificationType : uint8
 {
@@ -1293,9 +991,15 @@ struct FPolygonGroupToCreate
 	// @todo mesheditor: Material, bEnableCollision and bCastShadow should all be implemented as PolygonGroup attributes.
 	// This would require making Material an asset name, rather than a UObject*, which is also safer.
 
-	/** Material to assign to the new polygon group */
+	/** Material asset path to assign to the new polygon group */
+	UPROPERTY( BlueprintReadWrite, Category="Editable Mesh" )
+	FString MaterialAsset;
+
 	UPROPERTY( BlueprintReadWrite, Category="Editable Mesh" )
-	class UMaterialInterface* Material;
+	FName MaterialSlotName;
+
+	UPROPERTY( BlueprintReadWrite, Category="Editable Mesh" )
+	FName ImportedMaterialSlotName;
 
 	/** Whether the new polygon group should have collision enabled */
 	UPROPERTY( BlueprintReadWrite, Category="Editable Mesh" )
@@ -1311,8 +1015,7 @@ struct FPolygonGroupToCreate
 
 	/** Default constructor */
 	FPolygonGroupToCreate()
-		: Material( nullptr )
-		, bEnableCollision( false )
+		: bEnableCollision( false )
 		, bCastShadow( false )
 		, OriginalPolygonGroupID( FPolygonGroupID::Invalid )
 	{
@@ -1321,10 +1024,42 @@ struct FPolygonGroupToCreate
 	FString ToString() const
 	{
 		return FString::Printf(
-			TEXT( "Material:%s, bEnableCollision:%s, bCastShadow:%s, OriginalPolygonGroupID:%s" ),
-			Material ? *((UObject*)Material )->GetName() : TEXT( "<none>" ),
+			TEXT( "Material:%s, MaterialSlotName:%s, ImportedMaterialSlotName:%s, bEnableCollision:%s, bCastShadow:%s, OriginalPolygonGroupID:%s" ),
+			*MaterialAsset,
+			*MaterialSlotName.ToString(),
+			*ImportedMaterialSlotName.ToString(),
 			*LexicalConversion::ToString( bEnableCollision ),
 			*LexicalConversion::ToString( bCastShadow ),
 			*OriginalPolygonGroupID.ToString() );
 	}
 };
+
+
+USTRUCT( BlueprintType )
+struct FPolygonGroupForPolygon
+{
+	GENERATED_BODY()
+
+	/** Polygon to assign to a new group */
+	UPROPERTY( BlueprintReadWrite, Category="Editable Mesh" )
+	FPolygonID PolygonID;
+
+	/** Polygon group to assign polygon to */
+	UPROPERTY( BlueprintReadWrite, Category="Editable Mesh" )
+	FPolygonGroupID PolygonGroupID;
+
+	/** Default constructor */
+	FPolygonGroupForPolygon()
+		: PolygonID( FPolygonID::Invalid )
+		, PolygonGroupID( FPolygonGroupID::Invalid )
+	{
+	}
+
+	FString ToString() const
+	{
+		return FString::Printf(
+			TEXT( "PolygonID:%s, PolygonGroupID:%s" ),
+			*PolygonID.ToString(),
+			*PolygonGroupID.ToString() );
+	}
+};
diff --git a/Engine/Plugins/Runtime/EditableMesh/Source/EditableMesh/StaticMeshEditableMeshFormat.cpp b/Engine/Plugins/Runtime/EditableMesh/Source/EditableMesh/StaticMeshEditableMeshFormat.cpp
--- a/Engine/Plugins/Runtime/EditableMesh/Source/EditableMesh/StaticMeshEditableMeshFormat.cpp
+++ b/Engine/Plugins/Runtime/EditableMesh/Source/EditableMesh/StaticMeshEditableMeshFormat.cpp
@@ -45,6 +45,8 @@ UEditableMesh* FStaticMeshEditableMeshFormat::MakeEditableMesh( UPrimitiveCompon
 	}
 
 	UEditableMesh* EditableMesh = NewObject<UEditableMesh>();
+	EditableMesh->MeshDescription = NewObject<UMeshDescription>( EditableMesh );
+
 	UEditableStaticMeshAdapter* EditableStaticMesh = NewObject<UEditableStaticMeshAdapter>( EditableMesh );
 	EditableMesh->Adapters.Add( EditableStaticMesh );
 
diff --git a/Engine/Source/Runtime/MeshDescription/MeshDescription.Build.cs b/Engine/Source/Runtime/MeshDescription/MeshDescription.Build.cs
new file mode 100644
--- /dev/null
+++ b/Engine/Source/Runtime/MeshDescription/MeshDescription.Build.cs
@@ -0,0 +1,22 @@
+// Copyright 1998-2017 Epic Games, Inc. All Rights Reserved.
+
+namespace UnrealBuildTool.Rules
+{
+	public class MeshDescription : ModuleRules
+	{
+		public MeshDescription(ReadOnlyTargetRules Target) : base(Target)
+		{
+            PrivateIncludePaths.Add("Runtime/MeshDescription/Private");
+            PublicIncludePaths.Add("Runtime/MeshDescription/Public");
+
+			PrivateDependencyModuleNames.AddRange(
+				new string[]
+				{
+					"Core",
+					"CoreUObject",
+                    "Engine"
+				}
+			);
+		}
+	}
+}
diff --git a/Engine/Source/Runtime/MeshDescription/Private/MeshDescription.cpp b/Engine/Source/Runtime/MeshDescription/Private/MeshDescription.cpp
new file mode 100644
--- /dev/null
+++ b/Engine/Source/Runtime/MeshDescription/Private/MeshDescription.cpp
@@ -0,0 +1,180 @@
+// Copyright 1998-2017 Epic Games, Inc. All Rights Reserved.
+
+#include "MeshDescription.h"
+
+
+UMeshDescription::UMeshDescription()
+{
+	// Add basic vertex attributes
+	VertexAttributes().RegisterAttribute<FVector>( "Position" );
+	VertexAttributes().RegisterAttribute<float>( "CornerSharpness" );
+
+	// Add basic vertex instance attributes
+	VertexInstanceAttributes().RegisterAttribute<FVector2D>( "TextureCoordinate", 2 );
+	VertexInstanceAttributes().RegisterAttribute<FVector>( "Normal" );
+	VertexInstanceAttributes().RegisterAttribute<FVector>( "Tangent" );
+	VertexInstanceAttributes().RegisterAttribute<float>( "BinormalSign" );
+	VertexInstanceAttributes().RegisterAttribute<FVector4>( "Color" );
+
+	// Add basic edge attributes
+	EdgeAttributes().RegisterAttribute<bool>( "EdgeIsHard" );
+	EdgeAttributes().RegisterAttribute<float>( "CreaseSharpness" );
+
+	// Add basic polygon attributes
+	PolygonAttributes().RegisterAttribute<FVector>( "PolygonNormal" );
+	PolygonAttributes().RegisterAttribute<FVector>( "PolygonTangent" );
+	PolygonAttributes().RegisterAttribute<FVector>( "PolygonBinormal" );
+	PolygonAttributes().RegisterAttribute<FVector>( "PolygonCenter" );
+
+	// Add basic polygon group attributes
+	PolygonGroupAttributes().RegisterAttribute<FStringAssetReference>( "MaterialAsset" );
+	PolygonGroupAttributes().RegisterAttribute<FName>( "MaterialSlotName" );
+	PolygonGroupAttributes().RegisterAttribute<FName>( "ImportedMaterialSlotName" );
+	PolygonGroupAttributes().RegisterAttribute<bool>( "EnableCollision" );
+	PolygonGroupAttributes().RegisterAttribute<bool>( "CastShadow" );
+}
+
+
+void UMeshDescription::Serialize( FArchive& Ar )
+{
+	Super::Serialize( Ar );
+
+	Ar << VertexArray;
+	Ar << VertexInstanceArray;
+	Ar << EdgeArray;
+	Ar << PolygonArray;
+	Ar << PolygonGroupArray;
+
+	Ar << VertexAttributesSet;
+	Ar << VertexInstanceAttributesSet;
+	Ar << EdgeAttributesSet;
+	Ar << PolygonAttributesSet;
+	Ar << PolygonGroupAttributesSet;
+}
+
+
+void UMeshDescription::PostLoad()
+{
+}
+
+
+void UMeshDescription::Compact( FElementIDRemappings& OutRemappings )
+{
+	VertexArray.Compact( OutRemappings.NewVertexIndexLookup );
+	VertexInstanceArray.Compact( OutRemappings.NewVertexInstanceIndexLookup );
+	EdgeArray.Compact( OutRemappings.NewEdgeIndexLookup );
+	PolygonArray.Compact( OutRemappings.NewPolygonIndexLookup );
+	PolygonGroupArray.Compact( OutRemappings.NewPolygonGroupIndexLookup );
+
+	// @todo mesh description: compact attribute arrays according to these remappings
+	//
+
+	FixUpElementIDs( OutRemappings );
+}
+
+
+void UMeshDescription::Remap( const FElementIDRemappings& Remappings )
+{
+	VertexArray.Remap( Remappings.NewVertexIndexLookup );
+	VertexInstanceArray.Remap( Remappings.NewVertexInstanceIndexLookup );
+	EdgeArray.Remap( Remappings.NewEdgeIndexLookup );
+	PolygonArray.Remap( Remappings.NewPolygonIndexLookup );
+	PolygonGroupArray.Remap( Remappings.NewPolygonGroupIndexLookup );
+
+	// @todo mesh description: remap attribute arrays according to these remappings
+	//
+
+	FixUpElementIDs( Remappings );
+}
+
+
+void UMeshDescription::FixUpElementIDs( const FElementIDRemappings& Remappings )
+{
+	for( const FVertexID VertexID : VertexArray.GetElementIDs() )
+	{
+		FMeshVertex& Vertex = VertexArray[ VertexID ];
+
+		// Fix up vertex instance index references in vertices array
+		for( FVertexInstanceID& VertexInstanceID : Vertex.VertexInstanceIDs )
+		{
+			VertexInstanceID = Remappings.GetRemappedVertexInstanceID( VertexInstanceID );
+		}
+
+		// Fix up edge index references in the vertex array
+		for( FEdgeID& EdgeID : Vertex.ConnectedEdgeIDs )
+		{
+			EdgeID = Remappings.GetRemappedEdgeID( EdgeID );
+		}
+	}
+
+	// Fix up vertex index references in vertex instance array
+	for( const FVertexInstanceID VertexInstanceID : VertexInstanceArray.GetElementIDs() )
+	{
+		FMeshVertexInstance& VertexInstance = VertexInstanceArray[ VertexInstanceID ];
+
+		VertexInstance.VertexID = Remappings.GetRemappedVertexID( VertexInstance.VertexID );
+
+		for( FPolygonID& PolygonID : VertexInstance.ConnectedPolygons )
+		{
+			PolygonID = Remappings.GetRemappedPolygonID( PolygonID );
+		}
+	}
+
+	for( const FEdgeID EdgeID : EdgeArray.GetElementIDs() )
+	{
+		FMeshEdge& Edge = EdgeArray[ EdgeID ];
+
+		// Fix up vertex index references in Edges array
+		for( int32 Index = 0; Index < 2; Index++ )
+		{
+			Edge.VertexIDs[ Index ] = Remappings.GetRemappedVertexID( Edge.VertexIDs[ Index ] );
+		}
+
+		// Fix up references to section indices
+		for( FPolygonID& ConnectedPolygon : Edge.ConnectedPolygons )
+		{
+			ConnectedPolygon = Remappings.GetRemappedPolygonID( ConnectedPolygon );
+		}
+	}
+
+	for( const FPolygonID PolygonID : PolygonArray.GetElementIDs() )
+	{
+		FMeshPolygon& Polygon = PolygonArray[ PolygonID ];
+
+		// Fix up references to vertex indices in section polygons' contours
+		for( FVertexInstanceID& VertexInstanceID : Polygon.PerimeterContour.VertexInstanceIDs )
+		{
+			VertexInstanceID = Remappings.GetRemappedVertexInstanceID( VertexInstanceID );
+		}
+
+		for( FMeshPolygonContour& HoleContour : Polygon.HoleContours )
+		{
+			for( FVertexInstanceID& VertexInstanceID : HoleContour.VertexInstanceIDs )
+			{
+				VertexInstanceID = Remappings.GetRemappedVertexInstanceID( VertexInstanceID );
+			}
+		}
+
+		for( FMeshTriangle& Triangle : Polygon.Triangles )
+		{
+			for( int32 TriangleVertexNumber = 0; TriangleVertexNumber < 3; ++TriangleVertexNumber )
+			{
+				const FVertexInstanceID OriginalVertexInstanceID = Triangle.GetVertexInstanceID( TriangleVertexNumber );
+				const FVertexInstanceID NewVertexInstanceID = Remappings.GetRemappedVertexInstanceID( OriginalVertexInstanceID );
+				Triangle.SetVertexInstanceID( TriangleVertexNumber, NewVertexInstanceID );
+			}
+		}
+
+		Polygon.PolygonGroupID = Remappings.GetRemappedPolygonGroupID( Polygon.PolygonGroupID );
+	}
+
+	for( const FPolygonGroupID PolygonGroupID : PolygonGroupArray.GetElementIDs() )
+	{
+		FMeshPolygonGroup& PolygonGroup = PolygonGroupArray[ PolygonGroupID ];
+
+		for( FPolygonID& Polygon : PolygonGroup.Polygons )
+		{
+			Polygon = Remappings.GetRemappedPolygonID( Polygon );
+		}
+	}
+}
diff --git a/Engine/Source/Runtime/MeshDescription/Private/MeshTypes.cpp b/Engine/Source/Runtime/MeshDescription/Private/MeshTypes.cpp
new file mode 100644
--- /dev/null
+++ b/Engine/Source/Runtime/MeshDescription/Private/MeshTypes.cpp
@@ -0,0 +1,23 @@
+// Copyright 1998-2017 Epic Games, Inc. All Rights Reserved.
+
+#include "MeshTypes.h"
+
+const FElementID FElementID::Invalid( TNumericLimits<uint32>::Max() );
+const FVertexID FVertexID::Invalid( TNumericLimits<uint32>::Max() );
+const FVertexInstanceID FVertexInstanceID::Invalid( TNumericLimits<uint32>::Max() );
+const FEdgeID FEdgeID::Invalid( TNumericLimits<uint32>::Max() );
+const FPolygonGroupID FPolygonGroupID::Invalid( TNumericLimits<uint32>::Max() );
+const FPolygonID FPolygonID::Invalid( TNumericLimits<uint32>::Max() );
+
+const FName UEditableMeshAttribute::VertexPositionName( "VertexPosition" );
+const FName UEditableMeshAttribute::VertexCornerSharpnessName( "VertexCornerSharpness" );
+const FName UEditableMeshAttribute::VertexNormalName( "VertexNormal" );
+const FName UEditableMeshAttribute::VertexTangentName( "VertexTangent" );
+const FName UEditableMeshAttribute::VertexBinormalSignName( "VertexBinormalSign" );
+const FName UEditableMeshAttribute::VertexTextureCoordinateName( "VertexTextureCoordinate" );
+const FName UEditableMeshAttribute::VertexColorName( "VertexColor" );
+const FName UEditableMeshAttribute::EdgeIsHardName( "EdgeIsHard" );
+const FName UEditableMeshAttribute::EdgeCreaseSharpnessName( "EdgeCreaseSharpness" );
+const FName UEditableMeshAttribute::PolygonNormalName( "PolygonNormal" );
+const FName UEditableMeshAttribute::PolygonCenterName( "PolygonCenter" );
+
diff --git a/Engine/Source/Runtime/MeshDescription/Public/MeshDescription.h b/Engine/Source/Runtime/MeshDescription/Public/MeshDescription.h
new file mode 100644
--- /dev/null
+++ b/Engine/Source/Runtime/MeshDescription/Public/MeshDescription.h
@@ -0,0 +1,595 @@
+// Copyright 1998-2017 Epic Games, Inc. All Rights Reserved.
+
+#pragma once
+
+#include "CoreMinimal.h"
+#include "MeshTypes.h"
+#include "MeshElementArray.h"
+#include "Materials/MaterialInterface.h"
+#include "MeshDescription.generated.h"
+
+
+USTRUCT()
+struct FMeshVertex
+{
+	GENERATED_BODY()
+
+	/** Position of the vertex. */
+	UPROPERTY()
+	FVector VertexPosition;
+
+	/** All of vertex instances which reference this vertex (for split vertex support) */
+	UPROPERTY()
+	TArray<FVertexInstanceID> VertexInstanceIDs;
+
+	/** The edges connected to this vertex */
+	UPROPERTY()
+	TArray<FEdgeID> ConnectedEdgeIDs;
+
+	/** When subdivisions are enabled, this controls how sharp the vertex is, between 0.0 and 1.0. */
+	UPROPERTY()
+	float CornerSharpness;	// @todo mesheditor subdiv: Not really used by static meshes at all.  Only for UEditableMeshes that use subdivision features.  Move elsewhere.?
+
+	/** Serializer */
+	friend FArchive& operator<<( FArchive& Ar, FMeshVertex& Vertex )
+	{
+		Ar << Vertex.VertexPosition;
+		Ar << Vertex.VertexInstanceIDs;
+		Ar << Vertex.ConnectedEdgeIDs;
+		Ar << Vertex.CornerSharpness;
+		return Ar;
+	}
+};
+
+
+USTRUCT()
+struct FMeshVertexInstance
+{
+	GENERATED_BODY()
+
+	/** The vertex this is instancing */
+	UPROPERTY()
+	FVertexID VertexID;
+
+	/** List of connected polygons */
+	UPROPERTY()
+	TArray<FPolygonID> ConnectedPolygons;
+
+	/** UVs for the vertex instance */
+	UPROPERTY()
+	TArray<FVector2D> VertexUVs;
+
+	/** Normal vector */
+	UPROPERTY()
+	FVector Normal;
+	
+	/** Tangent vector */
+	UPROPERTY()
+	FVector Tangent;
+
+	/** Basis determinant sign used to calculate the sense of the binormal */
+	UPROPERTY()
+	float BinormalSign;
+
+	/** Vertex color */
+	UPROPERTY()
+	FLinearColor Color;
+
+	/** Serializer */
+	friend FArchive& operator<<( FArchive& Ar, FMeshVertexInstance& VertexInstance )
+	{
+		Ar << VertexInstance.VertexID;
+		Ar << VertexInstance.ConnectedPolygons;
+		Ar << VertexInstance.VertexUVs;
+		Ar << VertexInstance.Normal;
+		Ar << VertexInstance.Tangent;
+		Ar << VertexInstance.BinormalSign;
+		Ar << VertexInstance.Color;
+		return Ar;
+	}
+};
+
+
+USTRUCT()
+struct FMeshEdge
+{
+	GENERATED_BODY()
+
+	/** IDs of the two editable mesh vertices that make up this edge.  The winding direction is not defined. */
+	UPROPERTY()
+	FVertexID VertexIDs[ 2 ];
+
+	/** The polygons that share this edge.  It's best if there are always only two polygons that share
+	    the edge, and those polygons are facing the same direction */
+	UPROPERTY()
+	TArray<FPolygonID> ConnectedPolygons;
+
+	/** Whether this edge is 'hard' or not, for the purpose of vertex normal and tangent generation */
+	UPROPERTY()
+	bool bIsHardEdge;
+
+	/** When subdivisions are enabled, this controls how sharp the creasing of this edge will be, between 0.0 and 1.0. */
+	UPROPERTY()
+	float CreaseSharpness;
+
+	/** Serializer */
+	friend FArchive& operator<<( FArchive& Ar, FMeshEdge& Edge )
+	{
+		Ar << Edge.VertexIDs[ 0 ];
+		Ar << Edge.VertexIDs[ 1 ];
+		Ar << Edge.ConnectedPolygons;
+		Ar << Edge.bIsHardEdge;
+		Ar << Edge.CreaseSharpness;
+		return Ar;
+	}
+};
+
+
+USTRUCT()
+struct FMeshPolygonContour
+{
+	GENERATED_BODY()
+
+	/** The ordered list of vertex instances which make up the polygon contour. The winding direction is counter-clockwise. */
+	UPROPERTY()
+	TArray<FVertexInstanceID> VertexInstanceIDs;
+
+	/** Serializer */
+	friend FArchive& operator<<( FArchive& Ar, FMeshPolygonContour& Contour )
+	{
+		Ar << Contour.VertexInstanceIDs;
+		return Ar;
+	}
+};
+
+
+USTRUCT( BlueprintType )
+struct FMeshTriangle
+{
+	GENERATED_BODY()
+
+	/** First vertex instance that makes up this triangle.  Indices must be ordered counter-clockwise. */
+	UPROPERTY( BlueprintReadWrite, Category="Editable Mesh" )
+	FVertexInstanceID VertexInstanceID0;
+
+	/** Second vertex instance that makes up this triangle.  Indices must be ordered counter-clockwise. */
+	UPROPERTY( BlueprintReadWrite, Category="Editable Mesh" )
+	FVertexInstanceID VertexInstanceID1;
+
+	/** Third vertex instance that makes up this triangle.  Indices must be ordered counter-clockwise. */
+	UPROPERTY( BlueprintReadWrite, Category="Editable Mesh" )
+	FVertexInstanceID VertexInstanceID2;
+
+	/** Gets the specified triangle vertex instance ID.  Pass an index between 0 and 2 inclusive. */
+	inline FVertexInstanceID GetVertexInstanceID( const int32 Index ) const
+	{
+		checkSlow( Index >= 0 && Index <= 2 );
+		return reinterpret_cast<const FVertexInstanceID*>( this )[ Index ];
+	}
+
+	/** Sets the specified triangle vertex instance ID.  Pass an index between 0 and 2 inclusive, and the new vertex instance ID to store. */
+	inline void SetVertexInstanceID( const int32 Index, const FVertexInstanceID NewVertexInstanceID )
+	{
+		checkSlow( Index >= 0 && Index <= 2 );
+		( reinterpret_cast<FVertexInstanceID*>( this )[ Index ] ) = NewVertexInstanceID;
+	}
+
+	/** Serializer */
+	friend FArchive& operator<<( FArchive& Ar, FMeshTriangle& Tri )
+	{
+		Ar << Tri.VertexInstanceID0;
+		Ar << Tri.VertexInstanceID1;
+		Ar << Tri.VertexInstanceID2;
+		return Ar;
+	}
+};
+
+
+USTRUCT()
+struct FMeshPolygon
+{
+	GENERATED_BODY()
+
+	/** The outer boundary edges of this polygon */
+	UPROPERTY()
+	FMeshPolygonContour PerimeterContour;
+
+	/** Optional inner contours of this polygon that define holes inside of the polygon.  For the geometry to
+	    be considered valid, the hole contours should reside within the boundary of the polygon perimeter contour, 
+		and must not overlap each other.  No "nesting" of polygons inside the holes is supported -- those are 
+		simply separate polygons */
+	UPROPERTY()
+	TArray<FMeshPolygonContour> HoleContours;
+
+	/** List of triangles which make up this polygon */
+	UPROPERTY()
+	TArray<FMeshTriangle> Triangles;
+
+	/** The polygon group which contains this polygon */
+	UPROPERTY()
+	FPolygonGroupID PolygonGroupID;
+
+	/** Cached normal */
+	UPROPERTY()
+	FVector PolygonNormal;
+
+	/** Cached tangent */
+	UPROPERTY()
+	FVector PolygonTangent;
+
+	/** Cached binormal */
+	UPROPERTY()
+	FVector PolygonBinormal;
+
+	/** Cached center */
+	UPROPERTY()
+	FVector PolygonCenter;
+
+	/** Serializer */
+	friend FArchive& operator<<( FArchive& Ar, FMeshPolygon& Polygon )
+	{
+		Ar << Polygon.PerimeterContour;
+		Ar << Polygon.HoleContours;
+		Ar << Polygon.Triangles;
+		Ar << Polygon.PolygonGroupID;
+		return Ar;
+	}
+};
+
+
+USTRUCT()
+struct FMeshPolygonGroup
+{
+	GENERATED_BODY()
+
+	/** The material for this mesh section */
+	UPROPERTY()
+	FStringAssetReference MaterialAsset;
+
+	/** A unique name identifying this polygon group */
+	UPROPERTY()
+	FName MaterialSlotName;
+
+	UPROPERTY()
+	FName ImportedMaterialSlotName;
+
+	/** If true, collision is enabled for this section. */
+	UPROPERTY()
+	bool bEnableCollision;
+
+	/** If true, this section will cast a shadow */
+	UPROPERTY()
+	bool bCastShadow;
+
+	/** All polygons in this group */
+	TArray<FPolygonID> Polygons;
+
+	/** Serializer */
+	friend FArchive& operator<<( FArchive& Ar, FMeshPolygonGroup& PolygonGroup )
+	{
+		Ar << PolygonGroup.MaterialAsset;
+		Ar << PolygonGroup.MaterialSlotName;
+		Ar << PolygonGroup.ImportedMaterialSlotName;
+		Ar << PolygonGroup.bEnableCollision;
+		Ar << PolygonGroup.bCastShadow;
+		Ar << PolygonGroup.Polygons;
+		return Ar;
+	}
+};
+
+
+/** Define container types */
+using FVertexArray = TMeshElementArray<FMeshVertex, FVertexID>;
+using FVertexInstanceArray = TMeshElementArray<FMeshVertexInstance, FVertexInstanceID>;
+using FEdgeArray = TMeshElementArray<FMeshEdge, FEdgeID>;
+using FPolygonArray = TMeshElementArray<FMeshPolygon, FPolygonID>;
+using FPolygonGroupArray = TMeshElementArray<FMeshPolygonGroup, FPolygonGroupID>;
+
+/**
+ * List of attribute types which are supported.
+ * We do this so we can automatically generate the attribute containers and their associated accessors with
+ * some template magic. Adding a new attribute type requires no extra code elsewhere in the class.
+ */
+using AttributeTypes = TTuple
+<
+	FVector4,
+	FVector,
+	FVector2D,
+	float,
+	int,
+	bool,
+	FName,
+	FStringAssetReference
+>;
+
+/** Define aliases for element attributes */
+template <typename AttributeType> using TVertexAttributeArray = TMeshElementArray<AttributeType, FVertexID>;
+template <typename AttributeType> using TVertexInstanceAttributeArray = TMeshElementArray<AttributeType, FVertexInstanceID>;
+template <typename AttributeType> using TEdgeAttributeArray = TMeshElementArray<AttributeType, FEdgeID>;
+template <typename AttributeType> using TPolygonAttributeArray = TMeshElementArray<AttributeType, FPolygonID>;
+template <typename AttributeType> using TPolygonGroupAttributeArray = TMeshElementArray<AttributeType, FPolygonGroupID>;
+
+
+/**
+  * Define type for an attribute container of a single type.
+  * Attributes have a name and an index, and are keyed on an ElementIDType.
+  *
+  * This implies the below data structure:
+  * A TMap keyed on the attribute name,
+  * yielding a TArray indexed by attribute index,
+  * yielding a TMeshElementArray keyed on an Element ID,
+  * yielding an item of type AttributeType.
+  *
+  * This looks complicated, but actually makes attribute lookup easy when we are interested in a particular attribute for many element IDs.
+  * By caching the TMeshElementArray arrived at by the attribute name and index, we have O(1) access to that attribute for all elements.
+  */
+template <typename AttributeType, typename ElementIDType>
+using TAttributesArray = TMap<FName, TArray<TMeshElementArray<AttributeType, ElementIDType>>>;
+
+
+/**
+ * Helper template which transforms a tuple of types into a tuple of TAttributesArrays of those types.
+ *
+ * We need to instance TAttributeArrays for each type in the AttributeTypes tuple.
+ * Then we can access the appropriate array (as long as we know what its index is).
+ *
+ * This template, given ElementIDType and TTuple<A, B>, will generate:
+ * TTuple<TAttributesArray<A, ElementIDType>, TAttributesArray<B, ElementIDType>>
+ */
+template <typename ElementIDType, typename Tuple>
+struct TMakeAttributesSet;
+
+template <typename ElementIDType, typename... TupleTypes>
+struct TMakeAttributesSet<ElementIDType, TTuple<TupleTypes...>>
+{
+	using Type = TTuple<TAttributesArray<TupleTypes, ElementIDType>...>;
+};
+
+
+/**
+ * Helper template which gets the tuple index of a given type from a given TTuple.
+ *
+ * Given Type = char, and Tuple = TTuple<int, float, char>,
+ * TTupleIndex<Type, Tuple>::Value will be 2.
+ */
+template <typename Type, typename Tuple>
+struct TTupleIndex;
+
+template <typename Type, typename... Types>
+struct TTupleIndex<Type, TTuple<Type, Types...>>
+{
+	static const uint32 Value = 0U;
+};
+
+template <typename Type, typename Head, typename... Tail>
+struct TTupleIndex<Type, TTuple<Head, Tail...>>
+{
+	static const uint32 Value = 1U + TTupleIndex<Type, TTuple<Tail...>>::Value;
+};
+
+
+/**
+ * This is the actual type of the attribute container.
+ * It is a TTuple of TAttributesArray<attributeType>.
+ */
+template <typename ElementIDType>
+struct TAttributesSet : public TMakeAttributesSet<ElementIDType, AttributeTypes>::Type
+{
+	using Super = typename TMakeAttributesSet<ElementIDType, AttributeTypes>::Type;
+
+	/**
+	 * Register a new attribute name with the given type (must be a member of the AttributeTypes tuple).
+	 *
+	 * Example of use:
+	 *
+	 *		VertexInstanceAttributes().RegisterAttribute<FVector2D>( "UV", 8 );
+	 *                        . . .
+	 *		TVertexInstanceAttributeArray<FVector2D>& UV0 = VertexInstanceAttributes().GetAttributes<FVector2D>( "UV", 0 );
+	 *		UV0[ VertexInstanceID ] = FVector2D( 1.0f, 1.0f );
+	 */
+	template <typename AttributeType>
+	void RegisterAttribute( const FName AttributeName, const int32 NumberOfIndices = 1 )
+	{
+		this->Get<TTupleIndex<AttributeType, AttributeTypes>::Value>().Emplace( AttributeName ).SetNum( NumberOfIndices );
+	}
+
+	/** Unregister an attribute name with the given type */
+	template <typename AttributeType>
+	void UnregisterAttribute( const FName AttributeName )
+	{
+		this->Get<TTupleIndex<AttributeType, AttributeTypes>::Value>().Remove( AttributeName );
+	}
+
+	/** Determines whether an attribute of the given type exists with the given name */
+	template <typename AttributeType>
+	bool HasAttribute( const FName AttributeName ) const
+	{
+		return this->Get<TTupleIndex<AttributeType, AttributeTypes>::Value>().Contains( AttributeName );
+	}
+
+	/**
+	 * Get a set of attributes with the given type, name and index.
+	 *
+	 * Example of use:
+	 *
+	 *		const TVertexAttributeArray<FVector>& VertexPositions = VertexAttributes().GetAttributes<FVector>( "Position" );
+	 *		for( const FVertexID VertexID : GetVertices().GetElementIDs() )
+	 *		{
+	 *			const FVector Position = VertexPositions[ VertexID ];
+	 *			DoSomethingWith( Position );
+	 *		}
+	 */
+	template <typename AttributeType>
+	TMeshElementArray<AttributeType, ElementIDType>& GetAttributes( const FName AttributeName, const int32 AttributeIndex = 0 )
+	{
+		// @todo mesh description: should this handle non-existent attribute names and indices gracefully?
+		return this->Get<TTupleIndex<AttributeType, AttributeTypes>::Value>().FindChecked( AttributeName )[ AttributeIndex ];
+	}
+
+	template <typename AttributeType>
+	const TMeshElementArray<AttributeType, ElementIDType>& GetAttributes( const FName AttributeName, const int32 AttributeIndex = 0 ) const
+	{
+		// @todo mesh description: should this handle non-existent attribute names and indices gracefully?
+		return this->Get<TTupleIndex<AttributeType, AttributeTypes>::Value>().FindChecked( AttributeName )[ AttributeIndex ];
+	}
+
+	/** Returns the number of indices for the attribute with the given name */
+	template <typename AttributeType>
+	int32 GetAttributeIndexCount( const FName AttributeName ) const
+	{
+		// @todo mesh description: should this handle non-existent attribute names and indices gracefully?
+		return this->Get<TTupleIndex<AttributeType, AttributeTypes>::Value>().FindChecked( AttributeName ).Num();
+	}
+
+	/** Returns an array of all the attribute names registered for this attribute type */
+	template <typename AttributeType, typename Allocator>
+	void GetAttributeNames( TArray<FName, Allocator>& OutAttributeNames ) const
+	{
+		this->Get<TTupleIndex<AttributeType, AttributeTypes>::Value>().GetKeys( OutAttributeNames );
+	}
+
+	template <typename AttributeType>
+	AttributeType GetAttribute( const ElementIDType ElementID, const FName AttributeName, const int32 AttributeIndex = 0 )
+	{
+		return this->Get<TTupleIndex<AttributeType, AttributeTypes>::Value>().FindChecked( AttributeName )[ AttributeIndex ][ ElementID ];
+	}
+
+	template <typename AttributeType>
+	AttributeType SetAttribute( const ElementIDType ElementID, const FName AttributeName, const int32 AttributeIndex, const AttributeType& AttributeValue )
+	{
+		this->Get<TTupleIndex<AttributeType, AttributeTypes>::Value>().FindChecked( AttributeName )[ AttributeIndex ][ ElementID ] = AttributeValue;
+	}
+
+	friend FArchive& operator<<( FArchive& Ar, TAttributesSet& AttributesSet )
+	{
+		Ar << static_cast<Super&>( AttributesSet );
+		return Ar;
+	}
+};
+
+
+/**
+ * This is a structure which holds the ID remappings returned by a Compact operation, or passed to a Remap operation.
+ */
+struct FElementIDRemappings
+{
+	TSparseArray<FVertexID> NewVertexIndexLookup;
+	TSparseArray<FVertexInstanceID> NewVertexInstanceIndexLookup;
+	TSparseArray<FEdgeID> NewEdgeIndexLookup;
+	TSparseArray<FPolygonID> NewPolygonIndexLookup;
+	TSparseArray<FPolygonGroupID> NewPolygonGroupIndexLookup;
+
+	FVertexID GetRemappedVertexID( FVertexID VertexID ) const
+	{
+		checkSlow( NewVertexIndexLookup.IsAllocated( VertexID.GetValue() ) );
+		return NewVertexIndexLookup[ VertexID.GetValue() ];
+	}
+
+	FVertexInstanceID GetRemappedVertexInstanceID( FVertexInstanceID VertexInstanceID ) const
+	{
+		checkSlow( NewVertexInstanceIndexLookup.IsAllocated( VertexInstanceID.GetValue() ) );
+		return NewVertexInstanceIndexLookup[ VertexInstanceID.GetValue() ];
+	}
+
+	FEdgeID GetRemappedEdgeID( FEdgeID EdgeID ) const
+	{
+		checkSlow( NewEdgeIndexLookup.IsAllocated( EdgeID.GetValue() ) );
+		return NewEdgeIndexLookup[ EdgeID.GetValue() ];
+	}
+
+	FPolygonID GetRemappedPolygonID( FPolygonID PolygonID ) const
+	{
+		checkSlow( NewPolygonIndexLookup.IsAllocated( PolygonID.GetValue() ) );
+		return NewPolygonIndexLookup[ PolygonID.GetValue() ];
+	}
+
+	FPolygonGroupID GetRemappedPolygonGroupID( FPolygonGroupID PolygonGroupID ) const
+	{
+		check( NewPolygonGroupIndexLookup.IsAllocated( PolygonGroupID.GetValue() ) );
+		return NewPolygonGroupIndexLookup[ PolygonGroupID.GetValue() ];
+	}
+};
+
+
+UCLASS()
+class MESHDESCRIPTION_API UMeshDescription : public UObject
+{
+public:
+
+	GENERATED_BODY()
+
+	UMeshDescription();
+
+	// UObject interface
+	virtual void Serialize( FArchive& Ar ) override;
+	virtual void PostLoad() override;
+
+	FVertexArray& Vertices() { return VertexArray; }
+	const FVertexArray& Vertices() const { return VertexArray; }
+
+	FMeshVertex& GetVertex( const FVertexID VertexID ) { return VertexArray[ VertexID ]; }
+	const FMeshVertex& GetVertex( const FVertexID VertexID ) const { return VertexArray[ VertexID ]; }
+
+	FVertexInstanceArray& VertexInstances() { return VertexInstanceArray; }
+	const FVertexInstanceArray& VertexInstances() const { return VertexInstanceArray; }
+
+	FMeshVertexInstance& GetVertexInstance( const FVertexInstanceID VertexInstanceID ) { return VertexInstanceArray[ VertexInstanceID ]; }
+	const FMeshVertexInstance& GetVertexInstance( const FVertexInstanceID VertexInstanceID ) const { return VertexInstanceArray[ VertexInstanceID ]; }
+
+	FEdgeArray& Edges() { return EdgeArray; }
+	const FEdgeArray& Edges() const { return EdgeArray; }
+
+	FMeshEdge& GetEdge( const FEdgeID EdgeID ) { return EdgeArray[ EdgeID ]; }
+	const FMeshEdge& GetEdge( const FEdgeID EdgeID ) const { return EdgeArray[ EdgeID ]; }
+
+	FPolygonArray& Polygons() { return PolygonArray; }
+	const FPolygonArray& Polygons() const { return PolygonArray; }
+
+	FMeshPolygon& GetPolygon( const FPolygonID PolygonID ) { return PolygonArray[ PolygonID ]; }
+	const FMeshPolygon& GetPolygon( const FPolygonID PolygonID ) const { return PolygonArray[ PolygonID ]; }
+
+	FPolygonGroupArray& PolygonGroups() { return PolygonGroupArray; }
+	const FPolygonGroupArray& PolygonGroups() const { return PolygonGroupArray; }
+
+	FMeshPolygonGroup& GetPolygonGroup( const FPolygonGroupID PolygonGroupID ) { return PolygonGroupArray[ PolygonGroupID ]; }
+	const FMeshPolygonGroup& GetPolygonGroup( const FPolygonGroupID PolygonGroupID ) const { return PolygonGroupArray[ PolygonGroupID ]; }
+
+	TAttributesSet<FVertexID>& VertexAttributes() { return VertexAttributesSet; }
+	const TAttributesSet<FVertexID>& VertexAttributes() const { return VertexAttributesSet; }
+
+	TAttributesSet<FVertexInstanceID>& VertexInstanceAttributes() { return VertexInstanceAttributesSet; }
+	const TAttributesSet<FVertexInstanceID>& VertexInstanceAttributes() const { return VertexInstanceAttributesSet; }
+
+	TAttributesSet<FEdgeID>& EdgeAttributes() { return EdgeAttributesSet; }
+	const TAttributesSet<FEdgeID>& EdgeAttributes() const { return EdgeAttributesSet; }
+
+	TAttributesSet<FPolygonID>& PolygonAttributes() { return PolygonAttributesSet; }
+	const TAttributesSet<FPolygonID>& PolygonAttributes() const { return PolygonAttributesSet; }
+
+	TAttributesSet<FPolygonGroupID>& PolygonGroupAttributes() { return PolygonGroupAttributesSet; }
+	const TAttributesSet<FPolygonGroupID>& PolygonGroupAttributes() const { return PolygonGroupAttributesSet; }
+
+	/** Compacts the data held in the mesh description, and returns an object describing how the IDs have been remapped. */
+	void Compact( FElementIDRemappings& OutRemappings );
+
+	/** Remaps the element IDs in the mesh description according to the passed in object */
+	void Remap( const FElementIDRemappings& Remappings );
+
+
+private:
+
+	/** Given a set of index remappings, fixes up references to element IDs */
+	void FixUpElementIDs( const FElementIDRemappings& Remappings );
+
+	FVertexArray VertexArray;
+	FVertexInstanceArray VertexInstanceArray;
+	FEdgeArray EdgeArray;
+	FPolygonArray PolygonArray;
+	FPolygonGroupArray PolygonGroupArray;
+
+	TAttributesSet<FVertexID> VertexAttributesSet;
+	TAttributesSet<FVertexInstanceID> VertexInstanceAttributesSet;
+	TAttributesSet<FEdgeID> EdgeAttributesSet;
+	TAttributesSet<FPolygonID> PolygonAttributesSet;
+	TAttributesSet<FPolygonGroupID> PolygonGroupAttributesSet;
+};
diff --git a/Engine/Source/Runtime/MeshDescription/Public/MeshElementArray.h b/Engine/Source/Runtime/MeshDescription/Public/MeshElementArray.h
new file mode 100644
--- /dev/null
+++ b/Engine/Source/Runtime/MeshDescription/Public/MeshElementArray.h
@@ -0,0 +1,303 @@
+// Copyright 1998-2017 Epic Games, Inc. All Rights Reserved.
+
+#pragma once
+
+#include "CoreMinimal.h"
+
+
+/**
+ * This defines the container used to hold mesh elements.
+ * Its important properties are that it acts as an associative container (i.e. an element can be obtained from
+ * a given index), and that insert/delete/find are cheap.
+ * The current implementation is as a TSparseArray, but we abstract it so that this can be changed later if
+ * required, e.g. a TMap might be desirable if we wished to maintain unique indices for the lifetime of the container.
+ */
+template <typename ElementType>
+class TMeshElementArrayBase
+{
+public:
+
+	/**
+	 * Custom serialization for TMeshElementArrayBase.
+	 * The default TSparseArray serialization also compacts all the elements, removing the gaps and changing the indices.
+	 * The indices are significant in editable meshes, hence this is a custom serializer which preserves them.
+	 */
+	friend FArchive& operator<<( FArchive& Ar, TMeshElementArrayBase& Array )
+	{
+		Array.Container.CountBytes( Ar );
+
+		if( Ar.IsLoading() )
+		{
+			// Load array
+			TBitArray<> AllocatedIndices;
+			Ar << AllocatedIndices;
+
+			Array.Container.Empty( AllocatedIndices.Num() );
+			for( auto It = TConstSetBitIterator<>( AllocatedIndices ); It; ++It )
+			{
+				Array.Container.Insert( It.GetIndex(), ElementType() );
+				Ar << Array.Container[ It.GetIndex() ];
+			}
+		}
+		else
+		{
+			// Save array
+			const int32 MaxIndex = Array.Container.GetMaxIndex();
+
+			// We have to build the TBitArray representing allocated indices by hand, as we don't have access to it from outside TSparseArray.
+			// @todo core: consider replacing TSparseArray serialization with this format.
+			TBitArray<> AllocatedIndices( false, MaxIndex );
+			for( int32 Index = 0; Index < MaxIndex; ++Index )
+			{
+				if( Array.Container.IsAllocated( Index ) )
+				{
+					AllocatedIndices[ Index ] = true;
+				}
+			}
+			Ar << AllocatedIndices;
+
+			for( auto It = Array.Container.CreateIterator(); It; ++It )
+			{
+				Ar << *It;
+			}
+		}
+
+		return Ar;
+	}
+
+protected:
+
+	/** The actual container, represented by a sparse array */
+	TSparseArray<ElementType> Container;
+};
+
+
+/**
+ * We prefer to access elements of the container via strongly-typed IDs.
+ * This derived class imposes this type safety.
+ */
+template <typename ElementType, typename ElementIDType>
+class TMeshElementArray : public TMeshElementArrayBase<ElementType>
+{
+	static_assert( TIsDerivedFrom<ElementIDType, FElementID>::IsDerived, "ElementIDType must be derived from FElementID" );
+
+	using TMeshElementArrayBase<ElementType>::Container;
+
+public:
+
+	/** Resets the container, optionally reserving space for elements to be added */
+	void Reset( const int32 Elements = 0 )
+	{
+		Container.Reset();
+		Container.Reserve( Elements );
+	}
+
+	/** Reserves space for the specified total number of elements */
+	void Reserve( const int32 Elements ) { Container.Reserve( Elements ); }
+
+	/** Add a new element at the next available index, and return the new ID */
+	ElementIDType Add() { return ElementIDType( Container.Add( ElementType() ) ); }
+
+	/** Add the provided element at the next available index, and return the new ID */
+	ElementIDType Add( typename TTypeTraits<ElementType>::ConstInitType Element ) { return ElementIDType( Container.Add( Element ) ); }
+
+	/** Add the provided element at the next available index, and return the ID */
+	ElementIDType Add( ElementType&& Element ) { return ElementIDType( Container.Add( Forward<ElementType>( Element ) ) ); }
+
+	/** Inserts a new element with the given ID */
+	ElementType& Insert( const ElementIDType ID )
+	{
+		Container.Insert( ID.GetValue(), ElementType() );
+		return Container[ ID.GetValue() ];
+	}
+
+	/** Inserts the provided element with the given ID */
+	ElementType& Insert( const ElementIDType ID, typename TTypeTraits<ElementType>::ConstInitType Element )
+	{
+		Container.Insert( ID.GetValue(), Element );
+		return Container[ ID.GetValue() ];
+	}
+
+	/** Inserts the provided element with the given ID */
+	ElementType& Insert( const ElementIDType ID, ElementType&& Element )
+	{
+		Container.Insert( ID.GetValue(), Forward<ElementType>( Element ) );
+		return Container[ ID.GetValue() ];
+	}
+
+	/** Removes the element with the given ID */
+	void Remove( const ElementIDType ID )
+	{
+		checkSlow( Container.IsAllocated( ID.GetValue() ) );
+		Container.RemoveAt( ID.GetValue() );
+	}
+
+	/** Returns the element with the given ID */
+	ElementType& operator[]( const ElementIDType ID )
+	{
+		checkSlow( Container.IsAllocated( ID.GetValue() ) );
+		return Container[ ID.GetValue() ];
+	}
+
+	const ElementType& operator[]( const ElementIDType ID ) const
+	{
+		checkSlow( Container.IsAllocated( ID.GetValue() ) );
+		return Container[ ID.GetValue() ];
+	}
+
+	/** Returns the number of elements in the container */
+	int32 Num() const { return Container.Num(); }
+
+	/** Returns the index after the last valid element */
+	int32 GetArraySize() const { return Container.GetMaxIndex(); }
+
+	/** Returns the first valid ID */
+	ElementIDType GetFirstValidID() const
+	{
+		return Container.Num() > 0 ?
+			ElementIDType( typename TSparseArray<ElementType>::TConstIterator( Container ).GetIndex() ) :
+			ElementIDType::Invalid;
+	}
+
+	/** Returns whether the given ID is valid or not */
+	bool IsValid( const ElementIDType ID ) const
+	{
+		return ID.GetValue() >= 0 && ID.GetValue() < Container.GetMaxIndex() && Container.IsAllocated( ID.GetValue() );
+	}
+
+	/** Serializer */
+	friend FArchive& operator<<( FArchive& Ar, TMeshElementArray& Array )
+	{
+		Ar << static_cast<TMeshElementArrayBase<ElementType>&>( Array );
+		return Ar;
+	}
+
+	/** Compacts elements and returns a remapping table */
+	void Compact( TSparseArray<ElementIDType>& OutIndexRemap );
+
+	/** Remaps elements according to the passed remapping table */
+	void Remap( const TSparseArray<ElementIDType>& IndexRemap );
+
+	/**
+	 * This is a special type of iterator which returns successive IDs of valid elements, rather than
+	 * the elements themselves.
+	 * It is designed to be used with a range-for:
+	 *
+	 *     for( const FVertexID VertexID : GetVertices().GetElementIDs() )
+	 *     {
+	 *         DoSomethingWith( VertexID );
+	 *     }
+	 */
+	class TElementIDs
+	{
+	public:
+
+		explicit TElementIDs( const TSparseArray<ElementType>& InArray )
+			: Array( InArray )
+		{}
+
+		class TConstIterator
+		{
+		public:
+
+			explicit TConstIterator( typename TSparseArray<ElementType>::TConstIterator&& It )
+				: Iterator( MoveTemp( It ) )
+			{}
+
+			TConstIterator& operator++()
+			{
+				++Iterator;
+				return *this;
+			}
+
+			ElementIDType operator*() const
+			{
+				return Iterator ? ElementIDType( Iterator.GetIndex() ) : ElementIDType::Invalid;
+			}
+
+			friend bool operator==( const TConstIterator& Lhs, const TConstIterator& Rhs )
+			{
+				return Lhs.Iterator == Rhs.Iterator;
+			}
+
+			friend bool operator!=( const TConstIterator& Lhs, const TConstIterator& Rhs )
+			{
+				return Lhs.Iterator != Rhs.Iterator;
+			}
+
+		private:
+
+			typename TSparseArray<ElementType>::TConstIterator Iterator;
+		};
+
+		TConstIterator CreateConstIterator() const
+		{
+			return TConstIterator( typename TSparseArray<ElementType>::TConstIterator( Array ) );
+		}
+
+	private:
+
+		friend TConstIterator begin( const TElementIDs& ArrayIDs )
+		{
+			return TConstIterator( begin( ArrayIDs.Array ) );
+		}
+
+		friend TConstIterator end( const TElementIDs& ArrayIDs )
+		{
+			return TConstIterator( end( ArrayIDs.Array ) );
+		}
+
+		const TSparseArray<ElementType>& Array;
+	};
+
+	/** Return iterable proxy object from container */
+	TElementIDs GetElementIDs() const { return TElementIDs( Container ); }
+};
+
+
+template <typename ElementType, typename ElementIDType>
+void TMeshElementArray<ElementType, ElementIDType>::Compact( TSparseArray<ElementIDType>& OutIndexRemap )
+{
+	static TSparseArray<ElementType> NewContainer;
+	NewContainer.Empty( Container.Num() );
+
+	OutIndexRemap.Empty( Container.GetMaxIndex() );
+
+	// Add valid elements into a new contiguous sparse array.  Note non-const iterator so we can move elements.
+	for( TSparseArray<ElementType>::TIterator It( Container ); It; ++It )
+	{
+		const int32 OldElementIndex = It.GetIndex();
+
+		// @todo mesheditor: implement TSparseArray::Add( ElementType&& ) to save this obscure approach
+		const int32 NewElementIndex = NewContainer.Add( ElementType() );
+		NewContainer[ NewElementIndex ] = MoveTemp( *It );
+
+		// Provide an O(1) lookup from old index to new index, used when patching up vertex references afterwards
+		OutIndexRemap.Insert( OldElementIndex, ElementIDType( NewElementIndex ) );
+	}
+
+	Container = MoveTemp( NewContainer );
+}
+
+
+template <typename ElementType, typename ElementIDType>
+void TMeshElementArray<ElementType, ElementIDType>::Remap( const TSparseArray<ElementIDType>& IndexRemap )
+{
+	static TSparseArray<ElementType> NewContainer;
+	NewContainer.Empty( IndexRemap.GetMaxIndex() );
+
+	// Add valid elements into a new contiguous sparse array.  Note non-const iterator so we can move elements.
+	for( TSparseArray<ElementType>::TIterator It( Container ); It; ++It )
+	{
+		const int32 OldElementIndex = It.GetIndex();
+
+		check( IndexRemap.IsAllocated( OldElementIndex ) );
+		const int32 NewElementIndex = IndexRemap[ OldElementIndex ].GetValue();
+
+		// @todo mesheditor: implement TSparseArray::Insert( ElementType&& ) to save this obscure approach
+		NewContainer.Insert( NewElementIndex, ElementType() );
+		NewContainer[ NewElementIndex ] = MoveTemp( *It );
+	}
+
+	Container = MoveTemp( NewContainer );
+}
diff --git a/Engine/Source/Runtime/MeshDescription/Public/MeshTypes.h b/Engine/Source/Runtime/MeshDescription/Public/MeshTypes.h
new file mode 100644
--- /dev/null
+++ b/Engine/Source/Runtime/MeshDescription/Public/MeshTypes.h
@@ -0,0 +1,309 @@
+// Copyright 1998-2017 Epic Games, Inc. All Rights Reserved.
+
+#pragma once
+
+#include "CoreMinimal.h"
+#include "MeshTypes.generated.h"
+
+
+// @todo mesheditor: Need comments
+
+USTRUCT( BlueprintType )
+struct FElementID	// @todo mesheditor script: BP doesn't have name spaces, so we might need a more specific display name, or just rename our various types
+{
+	GENERATED_BODY()
+
+	FElementID()
+	{
+	}
+
+	explicit FElementID( const int32 InitIDValue )
+		: IDValue( InitIDValue )
+	{
+	}
+
+	FORCEINLINE int32 GetValue() const
+	{
+		return IDValue;
+	}
+
+	FORCEINLINE bool operator==( const FElementID& Other ) const
+	{
+		return IDValue == Other.IDValue;
+	}
+
+	FORCEINLINE bool operator!=( const FElementID& Other ) const
+	{
+		return IDValue != Other.IDValue;
+	}
+
+	FString ToString() const
+	{
+		return ( IDValue == Invalid.GetValue() ) ? TEXT( "Invalid" ) : FString::Printf( TEXT( "%d" ), IDValue );
+	}
+
+	friend FArchive& operator<<( FArchive& Ar, FElementID& Element )
+	{
+		Ar << Element.IDValue;
+		return Ar;
+	}
+
+	/** Invalid element ID */
+	MESHDESCRIPTION_API static const FElementID Invalid;
+
+protected:
+
+	/** The actual mesh element index this ID represents.  Read-only. */
+	UPROPERTY( BlueprintReadOnly, Category="Editable Mesh" )
+	int32 IDValue;
+};
+
+
+USTRUCT( BlueprintType )
+struct FVertexID : public FElementID
+{
+	GENERATED_BODY()
+
+	FVertexID()
+	{
+	}
+
+	explicit FVertexID( const FElementID InitElementID )
+		: FElementID( InitElementID.GetValue() )
+	{
+	}
+
+	explicit FVertexID( const int32 InitIDValue )
+		: FElementID( InitIDValue )
+	{
+	}
+
+	FORCEINLINE friend uint32 GetTypeHash( const FVertexID& Other )
+	{
+		return GetTypeHash( Other.IDValue );
+	}
+
+	/** Invalid vertex ID */
+	MESHDESCRIPTION_API static const FVertexID Invalid;
+};
+
+
+USTRUCT( BlueprintType )
+struct FVertexInstanceID : public FElementID
+{
+	GENERATED_BODY()
+
+	FVertexInstanceID()
+	{
+	}
+
+	explicit FVertexInstanceID( const FElementID InitElementID )
+		: FElementID( InitElementID.GetValue() )
+	{
+	}
+
+	explicit FVertexInstanceID( const uint32 InitIDValue )
+		: FElementID( InitIDValue )
+	{
+	}
+
+	FORCEINLINE friend uint32 GetTypeHash( const FVertexInstanceID& Other )
+	{
+		return GetTypeHash( Other.IDValue );
+	}
+
+	/** Invalid rendering vertex ID */
+	MESHDESCRIPTION_API static const FVertexInstanceID Invalid;
+};
+
+
+USTRUCT( BlueprintType )
+struct FEdgeID : public FElementID
+{
+	GENERATED_BODY()
+
+	FEdgeID()
+	{
+	}
+
+	explicit FEdgeID( const FElementID InitElementID )
+		: FElementID( InitElementID.GetValue() )
+	{
+	}
+
+	explicit FEdgeID( const int32 InitIDValue )
+		: FElementID( InitIDValue )
+	{
+	}
+
+	FORCEINLINE friend uint32 GetTypeHash( const FEdgeID& Other )
+	{
+		return GetTypeHash( Other.IDValue );
+	}
+
+	/** Invalid edge ID */
+	MESHDESCRIPTION_API static const FEdgeID Invalid;
+};
+
+
+USTRUCT( BlueprintType )
+struct FPolygonGroupID : public FElementID
+{
+	GENERATED_BODY()
+
+	FPolygonGroupID()
+	{
+	}
+
+	explicit FPolygonGroupID( const FElementID InitElementID )
+		: FElementID( InitElementID.GetValue() )
+	{
+	}
+
+	explicit FPolygonGroupID( const int32 InitIDValue )
+		: FElementID( InitIDValue )
+	{
+	}
+
+	FORCEINLINE friend uint32 GetTypeHash( const FPolygonGroupID& Other )
+	{
+		return GetTypeHash( Other.IDValue );
+	}
+
+	/** Invalid section ID */
+	MESHDESCRIPTION_API static const FPolygonGroupID Invalid;
+};
+
+
+USTRUCT( BlueprintType )
+struct FPolygonID : public FElementID
+{
+	GENERATED_BODY()
+
+	FPolygonID()
+	{
+	}
+
+	explicit FPolygonID( const FElementID InitElementID )
+		: FElementID( InitElementID.GetValue() )
+	{
+	}
+
+	explicit FPolygonID( const int32 InitIDValue )
+		: FElementID( InitIDValue )
+	{
+	}
+
+	FORCEINLINE friend uint32 GetTypeHash( const FPolygonID& Other )
+	{
+		return GetTypeHash( Other.IDValue );
+	}
+
+	/** Invalid polygon ID */
+	MESHDESCRIPTION_API static const FPolygonID Invalid;	// @todo mesheditor script: Can we expose these to BP nicely?	Do we even need to?
+};
+
+
+UCLASS( abstract )
+class MESHDESCRIPTION_API UEditableMeshAttribute : public UObject
+{
+	GENERATED_BODY()
+
+public:
+
+	//
+	// Vertex data for any vertex
+	//
+
+	/** Static: The attribute name for vertex position */
+	UFUNCTION( BlueprintPure, Category="Editable Mesh" ) static inline FName VertexPosition()
+	{
+		return VertexPositionName;
+	}
+
+	/** Static: The attribute name for vertex corner sharpness (only applies to subdivision meshes) */
+	UFUNCTION( BlueprintPure, Category="Editable Mesh" ) static inline FName VertexCornerSharpness()
+	{
+		return VertexCornerSharpnessName;
+	}
+
+	//
+	// Vertex instance data
+	//
+
+	/** Static: The attribute name for vertex normal (tangent Z) */
+	UFUNCTION( BlueprintPure, Category="Editable Mesh" ) static inline FName VertexNormal()
+	{
+		return VertexNormalName;
+	}
+
+	/** Static: The attribute name for vertex tangent vector (tangent X) */
+	UFUNCTION( BlueprintPure, Category="Editable Mesh" ) static inline FName VertexTangent()
+	{
+		return VertexTangentName;
+	}
+
+	/** Static: The attribute name for the vertex basis determinant sign (used to calculate the direction of tangent Y) */
+	UFUNCTION( BlueprintPure, Category="Editable Mesh" ) static inline FName VertexBinormalSign()
+	{
+		return VertexBinormalSignName;
+	}
+
+	/** Static: The attribute name for vertex texture coordinate.  The attribute index defines which texture coordinate set. */
+	UFUNCTION( BlueprintPure, Category="Editable Mesh" ) static inline FName VertexTextureCoordinate()
+	{
+		return VertexTextureCoordinateName;
+	}
+
+	/** Static: The attribute name for the vertex color. */
+	UFUNCTION( BlueprintPure, Category="Editable Mesh" ) static inline FName VertexColor()
+	{
+		return VertexColorName;
+	}
+
+	//
+	// Edges
+	//
+
+	/** Static: The attribute name for edge hardedness */
+	UFUNCTION( BlueprintPure, Category="Editable Mesh" ) static inline FName EdgeIsHard()
+	{
+		return EdgeIsHardName;
+	}
+
+	/** Static: The attribute name for edge crease sharpness (only applies to subdivision meshes) */
+	UFUNCTION( BlueprintPure, Category="Editable Mesh" ) static inline FName EdgeCreaseSharpness()
+	{
+		return EdgeCreaseSharpnessName;
+	}
+
+	//
+	// Polygons
+	//
+
+	/** Static: The attribute name for polygon normal */
+	UFUNCTION( BlueprintPure, Category="Editable Mesh" ) static inline FName PolygonNormal()
+	{
+		return PolygonNormalName;
+	}
+
+	/** Static: The attribute name for polygon center */
+	UFUNCTION( BlueprintPure, Category="Editable Mesh" ) static inline FName PolygonCenter()
+	{
+		return PolygonCenterName;
+	}
+
+private:
+
+	static const FName VertexPositionName;
+	static const FName VertexCornerSharpnessName;
+	static const FName VertexNormalName;
+	static const FName VertexTangentName;
+	static const FName VertexBinormalSignName;
+	static const FName VertexTextureCoordinateName;
+	static const FName VertexColorName;
+	static const FName EdgeIsHardName;
+	static const FName EdgeCreaseSharpnessName;
+	static const FName PolygonNormalName;
+	static const FName PolygonCenterName;
+};
