Added basic API to UMeshDescription for creating/deleting mesh elements, and correctly modifying internals.

diff --git a/Engine/Source/Runtime/MeshDescription/Public/MeshDescription.h b/Engine/Source/Runtime/MeshDescription/Public/MeshDescription.h
--- a/Engine/Source/Runtime/MeshDescription/Public/MeshDescription.h
+++ b/Engine/Source/Runtime/MeshDescription/Public/MeshDescription.h
@@ -14,6 +14,11 @@ struct FMeshVertex
 {
 	GENERATED_BODY()
 
+	FMeshVertex()
+		: VertexPosition( FVector::ZeroVector ),
+		  CornerSharpness( 0.0f )
+	{}
+
 	/** Position of the vertex. */
 	UPROPERTY()
 	FVector VertexPosition;
@@ -47,6 +52,14 @@ struct FMeshVertexInstance
 {
 	GENERATED_BODY()
 
+	FMeshVertexInstance()
+		: VertexID( FVertexID::Invalid ),
+		  Normal( FVector::ZeroVector ),
+		  Tangent( FVector::ZeroVector ),
+		  BinormalSign( 0.0f ),
+		  Color( FLinearColor::White )
+	{}
+
 	/** The vertex this is instancing */
 	UPROPERTY()
 	FVertexID VertexID;
@@ -95,6 +108,12 @@ struct FMeshEdge
 {
 	GENERATED_BODY()
 
+	FMeshEdge()
+		: VertexIDs{ FVertexID::Invalid, FVertexID::Invalid },
+		  bIsHardEdge( false ),
+		  CreaseSharpness( 0.0f )
+	{}
+
 	/** IDs of the two editable mesh vertices that make up this edge.  The winding direction is not defined. */
 	UPROPERTY()
 	FVertexID VertexIDs[ 2 ];
@@ -148,6 +167,12 @@ struct FMeshTriangle
 {
 	GENERATED_BODY()
 
+	FMeshTriangle()
+		: VertexInstanceID0( FVertexInstanceID::Invalid ),
+		  VertexInstanceID1( FVertexInstanceID::Invalid ),
+		  VertexInstanceID2( FVertexInstanceID::Invalid )
+	{}
+
 	/** First vertex instance that makes up this triangle.  Indices must be ordered counter-clockwise. */
 	UPROPERTY( BlueprintReadWrite, Category="Editable Mesh" )
 	FVertexInstanceID VertexInstanceID0;
@@ -190,6 +215,14 @@ struct FMeshPolygon
 {
 	GENERATED_BODY()
 
+	FMeshPolygon()
+		: PolygonGroupID( FPolygonGroupID::Invalid ),
+		  PolygonNormal( FVector::ZeroVector ),
+		  PolygonTangent( FVector::ZeroVector ),
+		  PolygonBinormal( FVector::ZeroVector ),
+		  PolygonCenter( FVector::ZeroVector )
+	{}
+
 	/** The outer boundary edges of this polygon */
 	UPROPERTY()
 	FMeshPolygonContour PerimeterContour;
@@ -242,6 +275,11 @@ struct FMeshPolygonGroup
 {
 	GENERATED_BODY()
 
+	FMeshPolygonGroup()
+		: bEnableCollision( true ),
+		  bCastShadow( true )
+	{}
+
 	/** The material for this mesh section */
 	UPROPERTY()
 	FStringAssetReference MaterialAsset;
@@ -569,6 +607,247 @@ public:
 	TAttributesSet<FPolygonGroupID>& PolygonGroupAttributes() { return PolygonGroupAttributesSet; }
 	const TAttributesSet<FPolygonGroupID>& PolygonGroupAttributes() const { return PolygonGroupAttributesSet; }
 
+	/** Adds a new vertex to the mesh and returns its ID */
+	FVertexID CreateVertex()
+	{
+		return VertexArray.Add();
+	}
+
+	/** Adds a new vertex to the mesh with the given ID */
+	void CreateVertexWithID( const FVertexID VertexID )
+	{
+		VertexArray.Insert( VertexID );
+	}
+
+	/** Deletes a vertex from the mesh */
+	void DeleteVertex( const FVertexID VertexID )
+	{
+		check( VertexArray[ VertexID ].ConnectedEdgeIDs.Num() == 0 );
+		check( VertexArray[ VertexID ].VertexInstanceIDs.Num() == 0 );
+		VertexArray.Remove( VertexID );
+	}
+
+private:
+	inline void CreateVertexInstance_Internal( const FVertexInstanceID VertexInstanceID, const FVertexID VertexID )
+	{
+		VertexInstanceArray[ VertexInstanceID ].VertexID = VertexID;
+		check( !VertexArray[ VertexID ].VertexInstanceIDs.Contains( VertexInstanceID ) );
+		VertexArray[ VertexID ].VertexInstanceIDs.Add( VertexInstanceID );
+	}
+
+public:
+	/** Adds a new vertex instance to the mesh and returns its ID */
+	FVertexInstanceID CreateVertexInstance( const FVertexID VertexID )
+	{
+		const FVertexInstanceID VertexInstanceID = VertexInstanceArray.Add();
+		CreateVertexInstance_Internal( VertexInstanceID, VertexID );
+		return VertexInstanceID;
+	}
+
+	/** Adds a new vertex instance to the mesh with the given ID */
+	void CreateVertexInstanceWithID( const FVertexInstanceID VertexInstanceID, const FVertexID VertexID )
+	{
+		VertexInstanceArray.Insert( VertexInstanceID );
+		CreateVertexInstance_Internal( VertexInstanceID, VertexID );
+	}
+
+	/** Deletes a vertex instance from a mesh */
+	void DeleteVertexInstance( const FVertexInstanceID VertexInstanceID, TArray<FVertexID>* InOutOrphanedVerticesPtr = nullptr )
+	{
+		const FVertexID VertexID = VertexInstanceArray[ VertexInstanceID ].VertexID;
+		verify( VertexArray[ VertexID ].VertexInstanceIDs.Remove( VertexInstanceID ) == 1 );
+		if( InOutOrphanedVerticesPtr && VertexArray[ VertexID ].VertexInstanceIDs.Num() == 0 && VertexArray[ VertexID ].ConnectedEdgeIDs.Num() == 0 )
+		{
+			InOutOrphanedVerticesPtr->Add( VertexID );
+		}
+		VertexInstanceArray.Remove( VertexInstanceID );
+	}
+
+private:
+	void CreateEdge_Internal( const FEdgeID EdgeID, const FVertexID VertexID0, const FVertexID VertexID1, const TArray<FPolygonID>& ConnectedPolygons )
+	{
+		FMeshEdge& Edge = EdgeArray[ EdgeID ];
+		Edge.VertexIDs[ 0 ] = VertexID0;
+		Edge.VertexIDs[ 1 ] = VertexID1;
+		Edge.ConnectedPolygons = ConnectedPolygons;
+		VertexArray[ VertexID0 ].ConnectedEdgeIDs.AddUnique( EdgeID );
+		VertexArray[ VertexID1 ].ConnectedEdgeIDs.AddUnique( EdgeID );
+	}
+
+public:
+	/** Adds a new edge to the mesh and returns its ID */
+	FEdgeID CreateEdge( const FVertexID VertexID0, const FVertexID VertexID1, const TArray<FPolygonID>& ConnectedPolygons = TArray<FPolygonID>() )
+	{
+		const FEdgeID EdgeID = EdgeArray.Add();
+		CreateEdge_Internal( EdgeID, VertexID0, VertexID1, ConnectedPolygons );
+		return EdgeID;
+	}
+
+	/** Adds a new edge to the mesh with the given ID */
+	void CreateEdgeWithID( const FEdgeID EdgeID, const FVertexID VertexID0, const FVertexID VertexID1, const TArray<FPolygonID>& ConnectedPolygons = TArray<FPolygonID>() )
+	{
+		EdgeArray.Insert( EdgeID );
+		CreateEdge_Internal( EdgeID, VertexID0, VertexID1, ConnectedPolygons );
+	}
+
+	/** Deletes an edge from a mesh */
+	void DeleteEdge( const FEdgeID EdgeID, TArray<FVertexID>* InOutOrphanedVerticesPtr = nullptr )
+	{
+		FMeshEdge& Edge = EdgeArray[ EdgeID ];
+		for( const FVertexID EdgeVertexID : Edge.VertexIDs )
+		{
+			FMeshVertex& Vertex = VertexArray[ EdgeVertexID ];
+			verify( Vertex.ConnectedEdgeIDs.RemoveSingle( EdgeID ) == 1 );
+			if( InOutOrphanedVerticesPtr && Vertex.ConnectedEdgeIDs.Num() == 0 )
+			{
+				check( Vertex.VertexInstanceIDs.Num() == 0 );  // We must already have deleted any vertex instances
+				InOutOrphanedVerticesPtr->Add( EdgeVertexID );
+			}
+		}
+		EdgeArray.Remove( EdgeID );
+	}
+
+	/** Pair of IDs representing the vertex instance ID on a contour, and the edge ID which starts at that point, winding counter-clockwise */
+	struct FContourPoint
+	{
+		FVertexInstanceID VertexInstanceID;
+		FEdgeID EdgeID;
+	};
+
+private:
+	void CreatePolygonContour_Internal( const FPolygonID PolygonID, const TArray<FContourPoint>& ContourPoints, FMeshPolygonContour& Contour )
+	{
+		Contour.VertexInstanceIDs.Reset( ContourPoints.Num() );
+		for( const FContourPoint ContourPoint : ContourPoints )
+		{
+			const FVertexInstanceID VertexInstanceID = ContourPoint.VertexInstanceID;
+			const FEdgeID EdgeID = ContourPoint.EdgeID;
+
+			Contour.VertexInstanceIDs.Add( VertexInstanceID );
+			check( !VertexInstanceArray[ VertexInstanceID ].ConnectedPolygons.Contains( PolygonID ) );
+			VertexInstanceArray[ VertexInstanceID ].ConnectedPolygons.Add( PolygonID );
+
+			check( !EdgeArray[ EdgeID ].ConnectedPolygons.Contains( PolygonID ) );
+			EdgeArray[ EdgeID ].ConnectedPolygons.Add( PolygonID );
+		}
+	}
+
+	void CreatePolygon_Internal( const FPolygonID PolygonID, const TArray<FContourPoint>& Perimeter, const TArray<TArray<FContourPoint>>& Holes )
+	{
+		FMeshPolygon& Polygon = PolygonArray[ PolygonID ];
+		CreatePolygonContour_Internal( PolygonID, Perimeter, Polygon.PerimeterContour );
+		for( const TArray<FContourPoint>& Hole : Holes )
+		{
+			Polygon.HoleContours.Emplace();
+			CreatePolygonContour_Internal( PolygonID, Hole, Polygon.HoleContours.Last() );
+		}
+	}
+
+public:
+	/** Adds a new polygon to the mesh and returns its ID */
+	FPolygonID CreatePolygon( const TArray<FContourPoint>& Perimeter, const TArray<TArray<FContourPoint>>& Holes = TArray<TArray<FContourPoint>>() )
+	{
+		const FPolygonID PolygonID = PolygonArray.Add();
+		CreatePolygon_Internal( PolygonID, Perimeter, Holes );
+		return PolygonID;
+	}
+
+	/** Adds a new polygon to the mesh with the given ID */
+	void CreatePolygonWithID( const FPolygonID PolygonID, const TArray<FContourPoint>& Perimeter, const TArray<TArray<FContourPoint>>& Holes = TArray<TArray<FContourPoint>>() )
+	{
+		PolygonArray.Insert( PolygonID );
+		CreatePolygon_Internal( PolygonID, Perimeter, Holes );
+	}
+
+private:
+	void DeletePolygonContour_Internal( const FPolygonID PolygonID, const TArray<FVertexInstanceID>& VertexInstanceIDs, TArray<FEdgeID>* InOutOrphanedEdgesPtr, TArray<FVertexInstanceID>* InOutOrphanedVertexInstancesPtr )
+	{
+		FVertexInstanceID LastVertexInstanceID = VertexInstanceIDs.Last();
+		for( const FVertexInstanceID VertexInstanceID : VertexInstanceIDs )
+		{
+			FMeshVertexInstance& VertexInstance = VertexInstanceArray[ VertexInstanceID ];
+			verify( VertexInstance.ConnectedPolygons.Remove( PolygonID ) == 1 );
+
+			if( InOutOrphanedVertexInstancesPtr && VertexInstance.ConnectedPolygons.Num() == 0 )
+			{
+				InOutOrphanedVertexInstancesPtr->Add( VertexInstanceID );
+			}
+
+			const FEdgeID EdgeID = GetVertexPairEdge( VertexInstanceArray[ LastVertexInstanceID ].VertexID, VertexInstanceArray[ VertexInstanceID ].VertexID );
+			check( EdgeID != FEdgeID::Invalid );
+			
+			FMeshEdge& Edge = EdgeArray[ EdgeID ];
+			verify( Edge.ConnectedPolygons.Remove( PolygonID ) == 1 );
+
+			if( InOutOrphanedEdgesPtr && Edge.ConnectedPolygons.Num() == 0 )
+			{
+				InOutOrphanedEdgesPtr->Add( EdgeID );
+			}
+
+			LastVertexInstanceID = VertexInstanceID;
+		}
+	}
+
+public:
+	/** Deletes a polygon from the mesh */
+	void DeletePolygon( const FPolygonID PolygonID, TArray<FEdgeID>* InOutOrphanedEdgesPtr = nullptr, TArray<FVertexInstanceID>* InOutOrphanedVertexInstancesPtr = nullptr, TArray<FPolygonGroupID>* InOutOrphanedPolygonGroupsPtr = nullptr )
+	{
+		FMeshPolygon& Polygon = PolygonArray[ PolygonID ];
+		DeletePolygonContour_Internal( PolygonID, Polygon.PerimeterContour.VertexInstanceIDs, InOutOrphanedEdgesPtr, InOutOrphanedVertexInstancesPtr );
+
+		for( FMeshPolygonContour& HoleContour : Polygon.HoleContours )
+		{
+			DeletePolygonContour_Internal( PolygonID, HoleContour.VertexInstanceIDs, InOutOrphanedEdgesPtr, InOutOrphanedVertexInstancesPtr );
+		}
+
+		FMeshPolygonGroup& PolygonGroup = PolygonGroupArray[ Polygon.PolygonGroupID ];
+		verify( PolygonGroup.Polygons.Remove( PolygonID ) == 1 );
+
+		if( InOutOrphanedPolygonGroupsPtr && PolygonGroup.Polygons.Num() == 0 )
+		{
+			InOutOrphanedPolygonGroupsPtr->Add( Polygon.PolygonGroupID );
+		}
+
+		PolygonArray.Remove( PolygonID );
+	}
+
+	/** Adds a new polygon group to the mesh and returns its ID */
+	FPolygonGroupID CreatePolygonGroup()
+	{
+		return PolygonGroupArray.Add();
+	}
+
+	/** Adds a new polygon group to the mesh with the given ID */
+	void CreatePolygonGroupWithID( const FPolygonGroupID PolygonGroupID )
+	{
+		PolygonGroupArray.Insert( PolygonGroupID );
+	}
+
+	/** Deletes a polygon group from the mesh */
+	void DeletePolygonGroup( const FPolygonGroupID PolygonGroupID )
+	{
+		check( PolygonGroupArray[ PolygonGroupID ].Polygons.Num() == 0 );
+		PolygonGroupArray.Remove( PolygonGroupID );
+	}
+
+
+	/** Returns the edge ID defined by the two given vertex IDs, if there is one; otherwise FEdgeID::Invalid */
+	FEdgeID GetVertexPairEdge( const FVertexID VertexID0, const FVertexID VertexID1 ) const
+	{
+		for( const FEdgeID VertexConnectedEdgeID : VertexArray[ VertexID0 ].ConnectedEdgeIDs )
+		{
+			const FVertexID EdgeVertexID0 = EdgeArray[ VertexConnectedEdgeID ].VertexIDs[ 0 ];
+			const FVertexID EdgeVertexID1 = EdgeArray[ VertexConnectedEdgeID ].VertexIDs[ 1 ];
+			if( ( EdgeVertexID0 == VertexID0 && EdgeVertexID1 == VertexID1 ) || ( EdgeVertexID0 == VertexID1 && EdgeVertexID1 == VertexID0 ) )
+			{
+				return VertexConnectedEdgeID;
+			}
+		}
+
+		return FEdgeID::Invalid;
+	}
+
+
 	/** Compacts the data held in the mesh description, and returns an object describing how the IDs have been remapped. */
 	void Compact( FElementIDRemappings& OutRemappings );
 
