Total revamp of mesh element attribute model.

diff --git a/Engine/Source/Runtime/MeshDescription/Public/MeshAttributeArray.h b/Engine/Source/Runtime/MeshDescription/Public/MeshAttributeArray.h
new file mode 100644
--- /dev/null
+++ b/Engine/Source/Runtime/MeshDescription/Public/MeshAttributeArray.h
@@ -0,0 +1,540 @@
+// Copyright 1998-2017 Epic Games, Inc. All Rights Reserved.
+
+#pragma once
+
+#include "CoreMinimal.h"
+#include "MeshTypes.h"
+#include "MeshElementRemappings.h"
+
+/**
+ * List of attribute types which are supported.
+ * We do this so we can automatically generate the attribute containers and their associated accessors with
+ * some template magic. Adding a new attribute type requires no extra code elsewhere in the class.
+ */
+using AttributeTypes = TTuple
+<
+	FVector4,
+	FVector,
+	FVector2D,
+	float,
+	int,
+	bool,
+	FName,
+	UObject*
+>;
+
+
+/**
+ * This defines the container used to hold mesh element attributes of a particular name and index.
+ * It is a simple TArray, so that all attributes are packed contiguously for each element ID.
+ */
+template <typename ElementType>
+class TMeshAttributeArrayBase
+{
+public:
+
+	/**
+	 * Custom serialization for TMeshAttributeArrayBase.
+	 */
+	friend FArchive& operator<<( FArchive& Ar, TMeshAttributeArrayBase& Array )
+	{
+		Ar << Array.Container;
+		return Ar;
+	}
+
+protected:
+
+	/** Expands the array if necessary so that the passed element index is valid. Newly created elements will be assigned the default value. */
+	void Insert( const int32 Index, const ElementType& Default )
+	{
+		if( Index >= Container.Num() )
+		{
+			// If the index is off the end of the container, add as many elements as required to make it the last valid index.
+			int32 StartIndex = Container.AddUninitialized( Index + 1 - Container.Num() );
+			ElementType* Data = Container.GetData() + StartIndex;
+
+			// Construct added elements with the default value passed in
+			while( StartIndex <= Index )
+			{
+				new( Data ) ElementType( Default );
+				StartIndex++;
+				Data++;
+			}
+		}
+	}
+
+	/** The actual container, represented by a regular array */
+	TArray<ElementType> Container;
+};
+
+
+template <typename AttributeType, typename ElementIDType>
+class TAttributeIndicesArray;
+
+/**
+ * We prefer to access elements of the container via strongly-typed IDs.
+ * This derived class imposes this type safety.
+ */
+template <typename ElementType, typename ElementIDType>
+class TMeshAttributeArray : public TMeshAttributeArrayBase<ElementType>
+{
+	static_assert( TIsDerivedFrom<ElementIDType, FElementID>::IsDerived, "ElementIDType must be derived from FElementID" );
+
+	using TMeshAttributeArrayBase<ElementType>::Container;
+
+public:
+
+	/** Element accessors */
+	FORCEINLINE const ElementType& operator[]( const ElementIDType ElementID ) const { return Container[ ElementID.GetValue() ]; }
+	FORCEINLINE ElementType& operator[]( const ElementIDType ElementID ) { return Container[ ElementID.GetValue() ]; }
+
+	/** Return size of container */
+	FORCEINLINE int32 Num() const { return Container.Num(); }
+
+	/** Return base of data */
+	FORCEINLINE const ElementType* GetData() const { return Container.GetData(); }
+
+protected:
+
+	friend class TAttributeIndicesArray<ElementType, ElementIDType>;
+
+	/** Expands the array if necessary so that the passed element index is valid. Newly created elements will be assigned the default value. */
+	FORCEINLINE void Insert( const ElementIDType Index, const ElementType& Default )
+	{
+		TMeshAttributeArrayBase<ElementType>::Insert( Index.GetValue(), Default );
+	}
+
+	/** Serializer */
+	friend FORCEINLINE FArchive& operator<<( FArchive& Ar, TMeshAttributeArray& Array )
+	{
+		Ar << static_cast<TMeshAttributeArrayBase<ElementType>&>( Array );
+		return Ar;
+	}
+
+	/** Remaps elements according to the passed remapping table */
+	void Remap( const TSparseArray<ElementIDType>& IndexRemap, const ElementType& Default );
+};
+
+
+template <typename ElementType, typename ElementIDType>
+void TMeshAttributeArray<ElementType, ElementIDType>::Remap( const TSparseArray<ElementIDType>& IndexRemap, const ElementType& Default )
+{
+	TMeshAttributeArray NewAttributeArray;
+
+	for( TSparseArray<ElementIDType>::TConstIterator It( IndexRemap ); It; ++It )
+	{
+		const int32 OldElementIndex = It.GetIndex();
+		const ElementIDType NewElementIndex = IndexRemap[ OldElementIndex ];
+
+		NewAttributeArray.Insert( NewElementIndex, Default );
+		NewAttributeArray[ NewElementIndex ] = MoveTemp( Container[ OldElementIndex ] );
+	}
+
+	Container = MoveTemp( NewAttributeArray.Container );
+}
+
+
+/** Define aliases for element attributes */
+template <typename AttributeType> using TVertexAttributeArray = TMeshAttributeArray<AttributeType, FVertexID>;
+template <typename AttributeType> using TVertexInstanceAttributeArray = TMeshAttributeArray<AttributeType, FVertexInstanceID>;
+template <typename AttributeType> using TEdgeAttributeArray = TMeshAttributeArray<AttributeType, FEdgeID>;
+template <typename AttributeType> using TPolygonAttributeArray = TMeshAttributeArray<AttributeType, FPolygonID>;
+template <typename AttributeType> using TPolygonGroupAttributeArray = TMeshAttributeArray<AttributeType, FPolygonGroupID>;
+
+
+/**
+ * Flags specifying properties of an attribute
+ * @todo mesh description: this needs to be moved to a application-specific place;
+ * this code is too low-level to assume particular meanings for the flags.
+ */
+enum class EMeshAttributeFlags : uint32
+{
+	None				= 0,
+	Lerpable			= ( 1 << 0 ),
+	AutoGenerated		= ( 1 << 1 ),
+	Mergeable			= ( 1 << 2 )
+};
+
+ENUM_CLASS_FLAGS( EMeshAttributeFlags );
+
+
+/**
+ * This class represents a container for a named attribute on a mesh element.
+ * It contains an array of TMeshAttributeArrays, one per attribute index. 
+ */
+template <typename AttributeType, typename ElementIDType>
+class TAttributeIndicesArray
+{
+public:
+
+	using Type = AttributeType;
+
+	/** Default constructor - required so that it builds correctly */
+	TAttributeIndicesArray() = default;
+
+	/** Constructor */
+	TAttributeIndicesArray( const int32 NumberOfIndices, const AttributeType& InDefaultValue, const EMeshAttributeFlags InFlags )
+		: DefaultValue( InDefaultValue ),
+		  Flags( InFlags )
+	{
+		ArrayForIndices.SetNum( NumberOfIndices );
+	}
+
+	/** Insert the element at the given index */
+	FORCEINLINE void Insert( const ElementIDType ElementID )
+	{
+		for( TMeshAttributeArray<AttributeType, ElementIDType>& ArrayForIndex : ArrayForIndices )
+		{
+			ArrayForIndex.Insert( ElementID, DefaultValue );
+		}
+	}
+
+	/** Remove the element at the given index, replacing it with a default value */
+	FORCEINLINE void Remove( const ElementIDType ElementID )
+	{
+		for( TMeshAttributeArray<AttributeType, ElementIDType>& ArrayForIndex : ArrayForIndices )
+		{
+			ArrayForIndex[ ElementID ] = DefaultValue;
+		}
+	}
+
+	/** Return the TMeshAttributeArray corresponding to the given attribute index */
+	FORCEINLINE const TMeshAttributeArray<AttributeType, ElementIDType>& GetArrayForIndex( const int32 Index ) const { return ArrayForIndices[ Index ]; }
+	FORCEINLINE TMeshAttributeArray<AttributeType, ElementIDType>& GetArrayForIndex( const int32 Index ) { return ArrayForIndices[ Index ]; }
+
+	/** Return flags for this attribute type */
+	FORCEINLINE EMeshAttributeFlags GetFlags() const { return Flags; }
+
+	/** Return number of indices this attribute has */
+	FORCEINLINE int32 GetNumIndices() const { return ArrayForIndices.Num(); }
+
+	/** Sets number of indices this attribute has */
+	void SetNumIndices( const int32 NumIndices )
+	{
+		const int32 OriginalNumIndices = ArrayForIndices.Num();
+		ArrayForIndices.SetNum( NumIndices );
+
+		// If there is already at least one attribute index, and it is non-empty, ensure that newly added indices
+		// are initialized to the same size with the default value (they must all have equal size).
+		if( OriginalNumIndices > 0 && ArrayForIndices[ 0 ].Num() > 0 )
+		{
+			for( int32 Index = OriginalNumIndices; Index < NumIndices; ++Index )
+			{
+				ArrayForIndices[ Index ].Insert( ElementIDType( ArrayForIndices[ 0 ].Num() - 1 ), DefaultValue );
+			}
+		}
+	}
+
+	/** Remaps all attribute indices according to the passed mapping */
+	void Remap( const TSparseArray<ElementIDType>& IndexRemap )
+	{
+		for( TMeshAttributeArray<AttributeType, ElementIDType>& ArrayForIndex : ArrayForIndices )
+		{
+			ArrayForIndex.Remap( IndexRemap, DefaultValue );
+		}
+	}
+
+	/** Serializer */
+	friend FArchive& operator<<( FArchive& Ar, TAttributeIndicesArray& AttributesArray )
+	{
+		Ar << AttributesArray.ArrayForIndices;
+		Ar << AttributesArray.DefaultValue;
+		Ar << AttributesArray.Flags;
+		return Ar;
+	}
+
+private:
+	/** An array of MeshAttributeArrays, one per attribute index */
+	TArray<TMeshAttributeArray<AttributeType, ElementIDType>> ArrayForIndices;
+
+	/** The default value for an attribute of this name */
+	AttributeType DefaultValue;
+
+	/** Implementation-defined attribute name flags */
+	EMeshAttributeFlags Flags;
+};
+
+
+/** Define aliases for element attributes */
+template <typename AttributeType> using TVertexAttributeIndicesArray = TAttributeIndicesArray<AttributeType, FVertexID>;
+template <typename AttributeType> using TVertexInstanceAttributeIndicesArray = TAttributeIndicesArray<AttributeType, FVertexInstanceID>;
+template <typename AttributeType> using TEdgeAttributeIndicesArray = TAttributeIndicesArray<AttributeType, FEdgeID>;
+template <typename AttributeType> using TPolygonAttributeIndicesArray = TAttributeIndicesArray<AttributeType, FPolygonID>;
+template <typename AttributeType> using TPolygonGroupAttributeIndicesArray = TAttributeIndicesArray<AttributeType, FPolygonGroupID>;
+
+
+/**
+ * This alias maps an attribute name to a TAttributeIndicesArray, i.e. an array of MeshAttributeArrays, one per attribute index.
+ */
+template <typename AttributeType, typename ElementIDType>
+using TAttributesMap = TMap<FName, TAttributeIndicesArray<AttributeType, ElementIDType>>;
+
+
+/**
+ * Helper template which transforms a tuple of types into a tuple of TAttributesArrays of those types.
+ *
+ * We need to instance TAttributeArrays for each type in the AttributeTypes tuple.
+ * Then we can access the appropriate array (as long as we know what its index is).
+ *
+ * This template, given ElementIDType and TTuple<A, B>, will generate:
+ * TTuple<TAttributesArray<A, ElementIDType>, TAttributesArray<B, ElementIDType>>
+ */
+template <typename ElementIDType, typename Tuple>
+struct TMakeAttributesSet;
+
+template <typename ElementIDType, typename... TupleTypes>
+struct TMakeAttributesSet<ElementIDType, TTuple<TupleTypes...>>
+{
+	using Type = TTuple<TAttributesMap<TupleTypes, ElementIDType>...>;
+};
+
+
+/**
+ * Helper template which gets the tuple index of a given type from a given TTuple.
+ *
+ * Given Type = char, and Tuple = TTuple<int, float, char>,
+ * TTupleIndex<Type, Tuple>::Value will be 2.
+ */
+template <typename Type, typename Tuple>
+struct TTupleIndex;
+
+template <typename Type, typename... Types>
+struct TTupleIndex<Type, TTuple<Type, Types...>>
+{
+	static const uint32 Value = 0U;
+};
+
+template <typename Type, typename Head, typename... Tail>
+struct TTupleIndex<Type, TTuple<Head, Tail...>>
+{
+	static const uint32 Value = 1U + TTupleIndex<Type, TTuple<Tail...>>::Value;
+};
+
+
+/**
+ * This is the container for all attributes of a particular mesh element.
+ * It contains a TTuple of TAttributesMap, one per attribute type (FVector, float, bool, etc)
+ */
+template <typename ElementIDType>
+class TAttributesSet
+{
+public:
+	/**
+	 * Register a new attribute name with the given type (must be a member of the AttributeTypes tuple).
+	 *
+	 * Example of use:
+	 *
+	 *		VertexInstanceAttributes().RegisterAttribute<FVector2D>( "UV", 8 );
+	 *                        . . .
+	 *		TVertexInstanceAttributeArray<FVector2D>& UV0 = VertexInstanceAttributes().GetAttributes<FVector2D>( "UV", 0 );
+	 *		UV0[ VertexInstanceID ] = FVector2D( 1.0f, 1.0f );
+	 */
+	template <typename AttributeType>
+	void RegisterAttribute( const FName AttributeName, const int32 NumberOfIndices = 1, const AttributeType& Default = AttributeType(), const EMeshAttributeFlags Flags = EMeshAttributeFlags::None )
+	{
+		Container.Get<TTupleIndex<AttributeType, AttributeTypes>::Value>().Emplace(
+			AttributeName,
+			TAttributeIndicesArray<AttributeType, ElementIDType>( NumberOfIndices, Default, Flags )
+		);
+	}
+
+	/** Unregister an attribute name with the given type */
+	template <typename AttributeType>
+	void UnregisterAttribute( const FName AttributeName )
+	{
+		Container.Get<TTupleIndex<AttributeType, AttributeTypes>::Value>().Remove( AttributeName );
+	}
+
+	/** Determines whether an attribute of the given type exists with the given name */
+	template <typename AttributeType>
+	bool HasAttribute( const FName AttributeName ) const
+	{
+		return Container.Get<TTupleIndex<AttributeType, AttributeTypes>::Value>().Contains( AttributeName );
+	}
+
+	/**
+	 * Get an attribute array with the given type, name and index.
+	 *
+	 * Example of use:
+	 *
+	 *		const TVertexAttributeArray<FVector>& VertexPositions = VertexAttributes().GetAttributes<FVector>( "Position" );
+	 *		for( const FVertexID VertexID : GetVertices().GetElementIDs() )
+	 *		{
+	 *			const FVector Position = VertexPositions[ VertexID ];
+	 *			DoSomethingWith( Position );
+	 *		}
+	 */
+	template <typename AttributeType>
+	TMeshAttributeArray<AttributeType, ElementIDType>& GetAttributes( const FName AttributeName, const int32 AttributeIndex = 0 )
+	{
+		// @todo mesh description: should this handle non-existent attribute names and indices gracefully?
+		return Container.Get<TTupleIndex<AttributeType, AttributeTypes>::Value>().FindChecked( AttributeName ).GetArrayForIndex( AttributeIndex );
+	}
+
+	template <typename AttributeType>
+	const TMeshAttributeArray<AttributeType, ElementIDType>& GetAttributes( const FName AttributeName, const int32 AttributeIndex = 0 ) const
+	{
+		// @todo mesh description: should this handle non-existent attribute names and indices gracefully?
+		return Container.Get<TTupleIndex<AttributeType, AttributeTypes>::Value>().FindChecked( AttributeName ).GetArrayForIndex( AttributeIndex );
+	}
+
+	/**
+	 * Get a set of attribute arrays with the given type and name.
+	 *
+	 * Example of use:
+	 *
+	 *		const TArray<TVertexInstanceAttributeArray<FVector2D>>& UVs = VertexInstanceAttributes().GetAttributesSet<FVector2D>( "UV" );
+	 *		for( const FVertexInstanceID VertexInstanceID : GetVertexInstances().GetElementIDs() )
+	 *		{
+	 *			const FVector2D UV0 = UVs[ 0 ][ VertexInstanceID ];
+	 *			const FVector2D UV1 = UVs[ 1 ][ VertexInstanceID ];
+	 *			DoSomethingWith( UV0, UV1 );
+	 *		}
+	 */
+	template <typename AttributeType>
+	TAttributeIndicesArray<AttributeType, ElementIDType>& GetAttributesSet( const FName AttributeName )
+	{
+		// @todo mesh description: should this handle non-existent attribute names gracefully?
+		return Container.Get<TTupleIndex<AttributeType, AttributeTypes>::Value>().FindChecked( AttributeName );
+	}
+
+	template <typename AttributeType>
+	const TAttributeIndicesArray<AttributeType, ElementIDType>& GetAttributesSet( const FName AttributeName ) const
+	{
+		// @todo mesh description: should this handle non-existent attribute names gracefully?
+		return Container.Get<TTupleIndex<AttributeType, AttributeTypes>::Value>().FindChecked( AttributeName );
+	}
+
+	/** Returns the number of indices for the attribute with the given name */
+	template <typename AttributeType>
+	int32 GetAttributeIndexCount( const FName AttributeName ) const
+	{
+		// @todo mesh description: should this handle non-existent attribute names and indices gracefully?
+		return Container.Get<TTupleIndex<AttributeType, AttributeTypes>::Value>().FindChecked( AttributeName ).GetNumIndices();
+	}
+
+	/** Sets the number of indices for the attribute with the given name */
+	template <typename AttributeType>
+	void SetAttributeIndexCount( const FName AttributeName, const int32 NumIndices )
+	{
+		Container.Get<TTupleIndex<AttributeType, AttributeTypes>::Value>().FindChecked( AttributeName ).SetNumIndices( NumIndices );
+	}
+
+	/** Returns an array of all the attribute names registered for this attribute type */
+	template <typename AttributeType, typename Allocator>
+	void GetAttributeNames( TArray<FName, Allocator>& OutAttributeNames ) const
+	{
+		Container.Get<TTupleIndex<AttributeType, AttributeTypes>::Value>().GetKeys( OutAttributeNames );
+	}
+
+	template <typename AttributeType>
+	AttributeType GetAttribute( const ElementIDType ElementID, const FName AttributeName, const int32 AttributeIndex = 0 ) const
+	{
+		return Container.Get<TTupleIndex<AttributeType, AttributeTypes>::Value>().FindChecked( AttributeName ).GetArrayForIndex( AttributeIndex )[ ElementID ];
+	}
+
+	template <typename AttributeType>
+	void SetAttribute( const ElementIDType ElementID, const FName AttributeName, const int32 AttributeIndex, const AttributeType& AttributeValue )
+	{
+		Container.Get<TTupleIndex<AttributeType, AttributeTypes>::Value>().FindChecked( AttributeName ).GetArrayForIndex( AttributeIndex )[ ElementID ] = AttributeValue;
+	}
+
+	/** Inserts a default-initialized value for all attributes of the given ID */
+	void Insert( const ElementIDType ElementID )
+	{
+		auto AddElements = [ ElementID ]( auto& AttributesMap )
+		{
+			for( auto& AttributeNameAndIndicesArray : AttributesMap )
+			{
+				AttributeNameAndIndicesArray.Value.Insert( ElementID );
+			}
+		};
+
+		VisitTupleElements( Container, AddElements );
+	}
+
+	/** Removes all attributes with the given ID */
+	void Remove( const ElementIDType ElementID )
+	{
+		auto RemoveElements = [ ElementID ]( auto& AttributesMap )
+		{
+			for( auto& AttributeNameAndIndicesArray : AttributesMap )
+			{
+				AttributeNameAndIndicesArray.Value.Remove( ElementID );
+			}
+		};
+
+		VisitTupleElements( Container, RemoveElements );
+	}
+
+	/**
+	 * Call the supplied function on each attribute.
+	 * The prototype should be Func( const FName AttributeName, auto& AttributeIndicesArray );
+	 */
+	template <typename FuncType>
+	void ForEachAttributeIndicesArray( const FuncType& Func )
+	{
+		auto ForEach = [ &Func ]( auto& AttributesMap )
+		{
+			for( auto& AttributeNameAndIndicesArray : AttributesMap )
+			{
+				Func( AttributeNameAndIndicesArray.Key, AttributeNameAndIndicesArray.Value );
+			}
+		};
+
+		VisitTupleElements( Container, ForEach );
+	}
+
+	template <typename FuncType>
+	void ForEachAttributeIndicesArray( const FuncType& Func ) const
+	{
+		auto ForEach = [ &Func ]( const auto& AttributesMap )
+		{
+			for( const auto& AttributeNameAndIndicesArray : AttributesMap )
+			{
+				Func( AttributeNameAndIndicesArray.Key, AttributeNameAndIndicesArray.Value );
+			}
+		};
+
+		VisitTupleElements( Container, ForEach );
+	}
+
+	/** Applies the given remapping to the attributes set */
+	void Remap( const TSparseArray<ElementIDType>& IndexRemap )
+	{
+		auto RemapElements = [ &IndexRemap ]( auto& AttributesMap )
+		{
+			for( auto& AttributeNameAndIndicesArray : AttributesMap )
+			{
+				AttributeNameAndIndicesArray.Value.Remap( IndexRemap );
+			}
+		};
+
+		VisitTupleElements( Container, RemapElements );
+	}
+
+	/** Serializer */
+	friend FArchive& operator<<( FArchive& Ar, TAttributesSet& AttributesSet )
+	{
+		Ar << AttributesSet.Container;
+		return Ar;
+	}
+
+private:
+	/**
+	  * Define type for the entire attribute container.
+	  * We can have attributes of multiple types, each with a name and an arbitrary number of indices,
+	  * whose elements are indexed by an ElementIDType.
+	  *
+	  * This implies the below data structure:
+	  * A TTuple (one per attribute type) of
+	  * TMap keyed on the attribute name,
+	  * yielding a TArray indexed by attribute index,
+	  * yielding a TMeshAttributeArray indexed by an Element ID,
+	  * yielding an item of type AttributeType.
+	  *
+	  * This looks complicated, but actually makes attribute lookup easy when we are interested in a particular attribute for many element IDs.
+	  * By caching the TMeshAttributeArray arrived at by the attribute name and index, we have O(1) access to that attribute for all elements.
+	  */
+	typename TMakeAttributesSet<ElementIDType, AttributeTypes>::Type Container;
+};
diff --git a/Engine/Source/Runtime/MeshDescription/Public/MeshAttributes.h b/Engine/Source/Runtime/MeshDescription/Public/MeshAttributes.h
new file mode 100644
--- /dev/null
+++ b/Engine/Source/Runtime/MeshDescription/Public/MeshAttributes.h
@@ -0,0 +1,49 @@
+// Copyright 1998-2017 Epic Games, Inc. All Rights Reserved.
+
+#pragma once
+
+#include "CoreMinimal.h"
+
+// @todo mesh description: for now, these attributes are hardcoded in the MeshDescription module.
+// It would be better if these were defined in separate modules, one per mesh type which uses MeshDescription.
+// This would give an extensible and modular structure whereby a mesh type defines its valid attributes and adapters.
+namespace MeshAttribute
+{
+	namespace Vertex
+	{
+		extern MESHDESCRIPTION_API const FName Position;
+		extern MESHDESCRIPTION_API const FName CornerSharpness;
+	}
+
+	namespace VertexInstance
+	{
+		extern MESHDESCRIPTION_API const FName TextureCoordinate;
+		extern MESHDESCRIPTION_API const FName Normal;
+		extern MESHDESCRIPTION_API const FName Tangent;
+		extern MESHDESCRIPTION_API const FName BinormalSign;
+		extern MESHDESCRIPTION_API const FName Color;
+	}
+
+	namespace Edge
+	{
+		extern MESHDESCRIPTION_API const FName IsHard;
+		extern MESHDESCRIPTION_API const FName CreaseSharpness;
+	}
+
+	namespace Polygon
+	{
+		extern MESHDESCRIPTION_API const FName Normal;
+		extern MESHDESCRIPTION_API const FName Tangent;
+		extern MESHDESCRIPTION_API const FName Binormal;
+		extern MESHDESCRIPTION_API const FName Center;
+	}
+
+	namespace PolygonGroup
+	{
+		extern MESHDESCRIPTION_API const FName MaterialAsset;
+		extern MESHDESCRIPTION_API const FName MaterialSlotName;
+		extern MESHDESCRIPTION_API const FName ImportedMaterialSlotName;
+		extern MESHDESCRIPTION_API const FName EnableCollision;
+		extern MESHDESCRIPTION_API const FName CastShadow;
+	}
+}
diff --git a/Engine/Source/Runtime/MeshDescription/Public/MeshDescription.h b/Engine/Source/Runtime/MeshDescription/Public/MeshDescription.h
--- a/Engine/Source/Runtime/MeshDescription/Public/MeshDescription.h
+++ b/Engine/Source/Runtime/MeshDescription/Public/MeshDescription.h
@@ -7,6 +7,7 @@
 #include "CoreTypes.h"
 #include "MeshTypes.h"
 #include "MeshElementArray.h"
+#include "MeshAttributeArray.h"
 #include "MeshDescription.generated.h"
 
 
@@ -16,14 +17,8 @@ struct FMeshVertex
 	GENERATED_BODY()
 
 	FMeshVertex()
-		: VertexPosition( FVector::ZeroVector ),
-		  CornerSharpness( 0.0f )
 	{}
 
-	/** Position of the vertex. */
-	UPROPERTY()
-	FVector VertexPosition;
-
 	/** All of vertex instances which reference this vertex (for split vertex support) */
 	UPROPERTY()
 	TArray<FVertexInstanceID> VertexInstanceIDs;
@@ -32,17 +27,11 @@ struct FMeshVertex
 	UPROPERTY()
 	TArray<FEdgeID> ConnectedEdgeIDs;
 
-	/** When subdivisions are enabled, this controls how sharp the vertex is, between 0.0 and 1.0. */
-	UPROPERTY()
-	float CornerSharpness;	// @todo mesheditor subdiv: Not really used by static meshes at all.  Only for UEditableMeshes that use subdivision features.  Move elsewhere.?
-
 	/** Serializer */
 	friend FArchive& operator<<( FArchive& Ar, FMeshVertex& Vertex )
 	{
-		Ar << Vertex.VertexPosition;
 		Ar << Vertex.VertexInstanceIDs;
 		Ar << Vertex.ConnectedEdgeIDs;
-		Ar << Vertex.CornerSharpness;
 		return Ar;
 	}
 };
@@ -54,11 +43,7 @@ struct FMeshVertexInstance
 	GENERATED_BODY()
 
 	FMeshVertexInstance()
-		: VertexID( FVertexID::Invalid ),
-		  Normal( FVector::ZeroVector ),
-		  Tangent( FVector::ZeroVector ),
-		  BinormalSign( 0.0f ),
-		  Color( FLinearColor::White )
+		: VertexID( FVertexID::Invalid )
 	{}
 
 	/** The vertex this is instancing */
@@ -69,36 +54,11 @@ struct FMeshVertexInstance
 	UPROPERTY()
 	TArray<FPolygonID> ConnectedPolygons;
 
-	/** UVs for the vertex instance */
-	UPROPERTY()
-	TArray<FVector2D> VertexUVs;
-
-	/** Normal vector */
-	UPROPERTY()
-	FVector Normal;
-	
-	/** Tangent vector */
-	UPROPERTY()
-	FVector Tangent;
-
-	/** Basis determinant sign used to calculate the sense of the binormal */
-	UPROPERTY()
-	float BinormalSign;
-
-	/** Vertex color */
-	UPROPERTY()
-	FLinearColor Color;
-
 	/** Serializer */
 	friend FArchive& operator<<( FArchive& Ar, FMeshVertexInstance& VertexInstance )
 	{
 		Ar << VertexInstance.VertexID;
 		Ar << VertexInstance.ConnectedPolygons;
-		Ar << VertexInstance.VertexUVs;
-		Ar << VertexInstance.Normal;
-		Ar << VertexInstance.Tangent;
-		Ar << VertexInstance.BinormalSign;
-		Ar << VertexInstance.Color;
 		return Ar;
 	}
 };
@@ -110,10 +70,10 @@ struct FMeshEdge
 	GENERATED_BODY()
 
 	FMeshEdge()
-		: VertexIDs{ FVertexID::Invalid, FVertexID::Invalid },
-		  bIsHardEdge( false ),
-		  CreaseSharpness( 0.0f )
-	{}
+	{
+		VertexIDs[ 0 ] = FVertexID::Invalid;
+		VertexIDs[ 1 ] = FVertexID::Invalid;
+	}
 
 	/** IDs of the two editable mesh vertices that make up this edge.  The winding direction is not defined. */
 	UPROPERTY()
@@ -124,22 +84,12 @@ struct FMeshEdge
 	UPROPERTY()
 	TArray<FPolygonID> ConnectedPolygons;
 
-	/** Whether this edge is 'hard' or not, for the purpose of vertex normal and tangent generation */
-	UPROPERTY()
-	bool bIsHardEdge;
-
-	/** When subdivisions are enabled, this controls how sharp the creasing of this edge will be, between 0.0 and 1.0. */
-	UPROPERTY()
-	float CreaseSharpness;
-
 	/** Serializer */
 	friend FArchive& operator<<( FArchive& Ar, FMeshEdge& Edge )
 	{
 		Ar << Edge.VertexIDs[ 0 ];
 		Ar << Edge.VertexIDs[ 1 ];
 		Ar << Edge.ConnectedPolygons;
-		Ar << Edge.bIsHardEdge;
-		Ar << Edge.CreaseSharpness;
 		return Ar;
 	}
 };
@@ -217,11 +167,7 @@ struct FMeshPolygon
 	GENERATED_BODY()
 
 	FMeshPolygon()
-		: PolygonGroupID( FPolygonGroupID::Invalid ),
-		  PolygonNormal( FVector::ZeroVector ),
-		  PolygonTangent( FVector::ZeroVector ),
-		  PolygonBinormal( FVector::ZeroVector ),
-		  PolygonCenter( FVector::ZeroVector )
+		: PolygonGroupID( FPolygonGroupID::Invalid )
 	{}
 
 	/** The outer boundary edges of this polygon */
@@ -243,22 +189,6 @@ struct FMeshPolygon
 	UPROPERTY()
 	FPolygonGroupID PolygonGroupID;
 
-	/** Cached normal */
-	UPROPERTY()
-	FVector PolygonNormal;
-
-	/** Cached tangent */
-	UPROPERTY()
-	FVector PolygonTangent;
-
-	/** Cached binormal */
-	UPROPERTY()
-	FVector PolygonBinormal;
-
-	/** Cached center */
-	UPROPERTY()
-	FVector PolygonCenter;
-
 	/** Serializer */
 	friend FArchive& operator<<( FArchive& Ar, FMeshPolygon& Polygon )
 	{
@@ -277,42 +207,14 @@ struct FMeshPolygonGroup
 	GENERATED_BODY()
 
 	FMeshPolygonGroup()
-		: bEnableCollision( true ),
-		  bCastShadow( true )
 	{}
 
-	/** The material for this mesh section */
-	UPROPERTY()
-	FStringAssetReference MaterialAsset;
-
-	/** A unique name identifying this polygon group */
-	UPROPERTY()
-	FName MaterialSlotName;
-
-	UPROPERTY()
-	FName ImportedMaterialSlotName;
-
-	/** If true, collision is enabled for this section. */
-	UPROPERTY()
-	bool bEnableCollision;
-
-	/** If true, this section will cast a shadow */
-	UPROPERTY()
-	bool bCastShadow;
-
 	/** All polygons in this group */
 	TArray<FPolygonID> Polygons;
 
 	/** Serializer */
 	friend FArchive& operator<<( FArchive& Ar, FMeshPolygonGroup& PolygonGroup )
 	{
-		//TODO FStringAssetReference cannot be archive with FArchive it can be archive only with FArchiveUObject
-		//Ar << PolygonGroup.MaterialAsset;
-		PolygonGroup.MaterialAsset.SerializePath(Ar);
-		Ar << PolygonGroup.MaterialSlotName;
-		Ar << PolygonGroup.ImportedMaterialSlotName;
-		Ar << PolygonGroup.bEnableCollision;
-		Ar << PolygonGroup.bCastShadow;
 		Ar << PolygonGroup.Polygons;
 		return Ar;
 	}
@@ -326,238 +228,6 @@ using FEdgeArray = TMeshElementArray<FMeshEdge, FEdgeID>;
 using FPolygonArray = TMeshElementArray<FMeshPolygon, FPolygonID>;
 using FPolygonGroupArray = TMeshElementArray<FMeshPolygonGroup, FPolygonGroupID>;
 
-/**
- * List of attribute types which are supported.
- * We do this so we can automatically generate the attribute containers and their associated accessors with
- * some template magic. Adding a new attribute type requires no extra code elsewhere in the class.
- */
-using AttributeTypes = TTuple
-<
-	FVector4,
-	FVector,
-	FVector2D,
-	float,
-	int,
-	bool,
-	FName,
-	FStringAssetReference
->;
-
-/** Define aliases for element attributes */
-template <typename AttributeType> using TVertexAttributeArray = TMeshElementArray<AttributeType, FVertexID>;
-template <typename AttributeType> using TVertexInstanceAttributeArray = TMeshElementArray<AttributeType, FVertexInstanceID>;
-template <typename AttributeType> using TEdgeAttributeArray = TMeshElementArray<AttributeType, FEdgeID>;
-template <typename AttributeType> using TPolygonAttributeArray = TMeshElementArray<AttributeType, FPolygonID>;
-template <typename AttributeType> using TPolygonGroupAttributeArray = TMeshElementArray<AttributeType, FPolygonGroupID>;
-
-
-/**
-  * Define type for an attribute container of a single type.
-  * Attributes have a name and an index, and are keyed on an ElementIDType.
-  *
-  * This implies the below data structure:
-  * A TMap keyed on the attribute name,
-  * yielding a TArray indexed by attribute index,
-  * yielding a TMeshElementArray keyed on an Element ID,
-  * yielding an item of type AttributeType.
-  *
-  * This looks complicated, but actually makes attribute lookup easy when we are interested in a particular attribute for many element IDs.
-  * By caching the TMeshElementArray arrived at by the attribute name and index, we have O(1) access to that attribute for all elements.
-  */
-template <typename AttributeType, typename ElementIDType>
-using TAttributesArray = TMap<FName, TArray<TMeshElementArray<AttributeType, ElementIDType>>>;
-
-
-/**
- * Helper template which transforms a tuple of types into a tuple of TAttributesArrays of those types.
- *
- * We need to instance TAttributeArrays for each type in the AttributeTypes tuple.
- * Then we can access the appropriate array (as long as we know what its index is).
- *
- * This template, given ElementIDType and TTuple<A, B>, will generate:
- * TTuple<TAttributesArray<A, ElementIDType>, TAttributesArray<B, ElementIDType>>
- */
-template <typename ElementIDType, typename Tuple>
-struct TMakeAttributesSet;
-
-template <typename ElementIDType, typename... TupleTypes>
-struct TMakeAttributesSet<ElementIDType, TTuple<TupleTypes...>>
-{
-	using Type = TTuple<TAttributesArray<TupleTypes, ElementIDType>...>;
-};
-
-
-/**
- * Helper template which gets the tuple index of a given type from a given TTuple.
- *
- * Given Type = char, and Tuple = TTuple<int, float, char>,
- * TTupleIndex<Type, Tuple>::Value will be 2.
- */
-template <typename Type, typename Tuple>
-struct TTupleIndex;
-
-template <typename Type, typename... Types>
-struct TTupleIndex<Type, TTuple<Type, Types...>>
-{
-	static const uint32 Value = 0U;
-};
-
-template <typename Type, typename Head, typename... Tail>
-struct TTupleIndex<Type, TTuple<Head, Tail...>>
-{
-	static const uint32 Value = 1U + TTupleIndex<Type, TTuple<Tail...>>::Value;
-};
-
-
-/**
- * This is the actual type of the attribute container.
- * It is a TTuple of TAttributesArray<attributeType>.
- */
-template <typename ElementIDType>
-struct TAttributesSet : public TMakeAttributesSet<ElementIDType, AttributeTypes>::Type
-{
-	using Super = typename TMakeAttributesSet<ElementIDType, AttributeTypes>::Type;
-
-	/**
-	 * Register a new attribute name with the given type (must be a member of the AttributeTypes tuple).
-	 *
-	 * Example of use:
-	 *
-	 *		VertexInstanceAttributes().RegisterAttribute<FVector2D>( "UV", 8 );
-	 *                        . . .
-	 *		TVertexInstanceAttributeArray<FVector2D>& UV0 = VertexInstanceAttributes().GetAttributes<FVector2D>( "UV", 0 );
-	 *		UV0[ VertexInstanceID ] = FVector2D( 1.0f, 1.0f );
-	 */
-	template <typename AttributeType>
-	void RegisterAttribute( const FName AttributeName, const int32 NumberOfIndices = 1 )
-	{
-		this->Get<TTupleIndex<AttributeType, AttributeTypes>::Value>().Emplace( AttributeName ).SetNum( NumberOfIndices );
-	}
-
-	/** Unregister an attribute name with the given type */
-	template <typename AttributeType>
-	void UnregisterAttribute( const FName AttributeName )
-	{
-		this->Get<TTupleIndex<AttributeType, AttributeTypes>::Value>().Remove( AttributeName );
-	}
-
-	/** Determines whether an attribute of the given type exists with the given name */
-	template <typename AttributeType>
-	bool HasAttribute( const FName AttributeName ) const
-	{
-		return this->Get<TTupleIndex<AttributeType, AttributeTypes>::Value>().Contains( AttributeName );
-	}
-
-	/**
-	 * Get a set of attributes with the given type, name and index.
-	 *
-	 * Example of use:
-	 *
-	 *		const TVertexAttributeArray<FVector>& VertexPositions = VertexAttributes().GetAttributes<FVector>( "Position" );
-	 *		for( const FVertexID VertexID : GetVertices().GetElementIDs() )
-	 *		{
-	 *			const FVector Position = VertexPositions[ VertexID ];
-	 *			DoSomethingWith( Position );
-	 *		}
-	 */
-	template <typename AttributeType>
-	TMeshElementArray<AttributeType, ElementIDType>& GetAttributes( const FName AttributeName, const int32 AttributeIndex = 0 )
-	{
-		// @todo mesh description: should this handle non-existent attribute names and indices gracefully?
-		return this->Get<TTupleIndex<AttributeType, AttributeTypes>::Value>().FindChecked( AttributeName )[ AttributeIndex ];
-	}
-
-	template <typename AttributeType>
-	const TMeshElementArray<AttributeType, ElementIDType>& GetAttributes( const FName AttributeName, const int32 AttributeIndex = 0 ) const
-	{
-		// @todo mesh description: should this handle non-existent attribute names and indices gracefully?
-		return this->Get<TTupleIndex<AttributeType, AttributeTypes>::Value>().FindChecked( AttributeName )[ AttributeIndex ];
-	}
-
-	/** Returns the number of indices for the attribute with the given name */
-	template <typename AttributeType>
-	int32 GetAttributeIndexCount( const FName AttributeName ) const
-	{
-		// @todo mesh description: should this handle non-existent attribute names and indices gracefully?
-		return this->Get<TTupleIndex<AttributeType, AttributeTypes>::Value>().FindChecked( AttributeName ).Num();
-	}
-
-	/** Returns an array of all the attribute names registered for this attribute type */
-	template <typename AttributeType, typename Allocator>
-	void GetAttributeNames( TArray<FName, Allocator>& OutAttributeNames ) const
-	{
-		this->Get<TTupleIndex<AttributeType, AttributeTypes>::Value>().GetKeys( OutAttributeNames );
-	}
-
-	template <typename AttributeType>
-	AttributeType GetAttribute( const ElementIDType ElementID, const FName AttributeName, const int32 AttributeIndex = 0 )
-	{
-		return this->Get<TTupleIndex<AttributeType, AttributeTypes>::Value>().FindChecked( AttributeName )[ AttributeIndex ][ ElementID ];
-	}
-
-	template <typename AttributeType>
-	void SetAttribute( const ElementIDType ElementID, const FName AttributeName, const int32 AttributeIndex, const AttributeType& AttributeValue )
-	{
-		TMeshElementArray<AttributeType, ElementIDType>& MeshElementArray = this->Get<TTupleIndex<AttributeType, AttributeTypes>::Value>().FindChecked(AttributeName)[AttributeIndex];
-		if(!MeshElementArray.IsValid(ElementID))
-		{
-			MeshElementArray.Insert(ElementID, AttributeValue);
-		}
-		else
-		{
-			MeshElementArray[ElementID] = AttributeValue;
-		}
-	}
-
-	friend FArchive& operator<<( FArchive& Ar, TAttributesSet& AttributesSet )
-	{
-		Ar << static_cast<Super&>( AttributesSet );
-		return Ar;
-	}
-};
-
-
-/**
- * This is a structure which holds the ID remappings returned by a Compact operation, or passed to a Remap operation.
- */
-struct FElementIDRemappings
-{
-	TSparseArray<FVertexID> NewVertexIndexLookup;
-	TSparseArray<FVertexInstanceID> NewVertexInstanceIndexLookup;
-	TSparseArray<FEdgeID> NewEdgeIndexLookup;
-	TSparseArray<FPolygonID> NewPolygonIndexLookup;
-	TSparseArray<FPolygonGroupID> NewPolygonGroupIndexLookup;
-
-	FVertexID GetRemappedVertexID( FVertexID VertexID ) const
-	{
-		checkSlow( NewVertexIndexLookup.IsAllocated( VertexID.GetValue() ) );
-		return NewVertexIndexLookup[ VertexID.GetValue() ];
-	}
-
-	FVertexInstanceID GetRemappedVertexInstanceID( FVertexInstanceID VertexInstanceID ) const
-	{
-		checkSlow( NewVertexInstanceIndexLookup.IsAllocated( VertexInstanceID.GetValue() ) );
-		return NewVertexInstanceIndexLookup[ VertexInstanceID.GetValue() ];
-	}
-
-	FEdgeID GetRemappedEdgeID( FEdgeID EdgeID ) const
-	{
-		checkSlow( NewEdgeIndexLookup.IsAllocated( EdgeID.GetValue() ) );
-		return NewEdgeIndexLookup[ EdgeID.GetValue() ];
-	}
-
-	FPolygonID GetRemappedPolygonID( FPolygonID PolygonID ) const
-	{
-		checkSlow( NewPolygonIndexLookup.IsAllocated( PolygonID.GetValue() ) );
-		return NewPolygonIndexLookup[ PolygonID.GetValue() ];
-	}
-
-	FPolygonGroupID GetRemappedPolygonGroupID( FPolygonGroupID PolygonGroupID ) const
-	{
-		check( NewPolygonGroupIndexLookup.IsAllocated( PolygonGroupID.GetValue() ) );
-		return NewPolygonGroupIndexLookup[ PolygonGroupID.GetValue() ];
-	}
-};
 
 
 UCLASS()
@@ -617,16 +287,32 @@ public:
 	TAttributesSet<FPolygonGroupID>& PolygonGroupAttributes() { return PolygonGroupAttributesSet; }
 	const TAttributesSet<FPolygonGroupID>& PolygonGroupAttributes() const { return PolygonGroupAttributesSet; }
 
+	/** Reserves space for this number of new vertices */
+	void ReserveNewVertices( const int32 NumVertices )
+	{
+		VertexArray.Reserve( VertexArray.Num() + NumVertices );
+	}
+
+private:
+	void CreateVertex_Internal( const FVertexID VertexID )
+	{
+		VertexAttributesSet.Insert( VertexID );
+	}
+
+public:
 	/** Adds a new vertex to the mesh and returns its ID */
 	FVertexID CreateVertex()
 	{
-		return VertexArray.Add();
+		const FVertexID VertexID = VertexArray.Add();
+		CreateVertex_Internal( VertexID );
+		return VertexID;
 	}
 
 	/** Adds a new vertex to the mesh with the given ID */
 	void CreateVertexWithID( const FVertexID VertexID )
 	{
 		VertexArray.Insert( VertexID );
+		CreateVertex_Internal( VertexID );
 	}
 
 	/** Deletes a vertex from the mesh */
@@ -635,18 +321,26 @@ public:
 		check( VertexArray[ VertexID ].ConnectedEdgeIDs.Num() == 0 );
 		check( VertexArray[ VertexID ].VertexInstanceIDs.Num() == 0 );
 		VertexArray.Remove( VertexID );
+		VertexAttributesSet.Remove( VertexID );
 	}
 
 #if WITH_EDITORONLY_DATA
 	FString GetIdString();
 #endif
 
+	/** Reserves space for this number of new vertex instances */
+	void ReserveNewVertexInstances( const int32 NumVertexInstances )
+	{
+		VertexInstanceArray.Reserve( VertexInstanceArray.Num() + NumVertexInstances );
+	}
+
 private:
-	inline void CreateVertexInstance_Internal( const FVertexInstanceID VertexInstanceID, const FVertexID VertexID )
+	void CreateVertexInstance_Internal( const FVertexInstanceID VertexInstanceID, const FVertexID VertexID )
 	{
 		VertexInstanceArray[ VertexInstanceID ].VertexID = VertexID;
 		check( !VertexArray[ VertexID ].VertexInstanceIDs.Contains( VertexInstanceID ) );
 		VertexArray[ VertexID ].VertexInstanceIDs.Add( VertexInstanceID );
+		VertexInstanceAttributesSet.Insert( VertexInstanceID );
 	}
 
 public:
@@ -666,15 +360,23 @@ public:
 	}
 
 	/** Deletes a vertex instance from a mesh */
-	void DeleteVertexInstance( const FVertexInstanceID VertexInstanceID, TSet<FVertexID>* InOutOrphanedVerticesPtr = nullptr )
+	void DeleteVertexInstance( const FVertexInstanceID VertexInstanceID, TArray<FVertexID>* InOutOrphanedVerticesPtr = nullptr )
 	{
+		check( VertexInstanceArray[ VertexInstanceID ].ConnectedPolygons.Num() == 0 );
 		const FVertexID VertexID = VertexInstanceArray[ VertexInstanceID ].VertexID;
 		verify( VertexArray[ VertexID ].VertexInstanceIDs.Remove( VertexInstanceID ) == 1 );
 		if( InOutOrphanedVerticesPtr && VertexArray[ VertexID ].VertexInstanceIDs.Num() == 0 && VertexArray[ VertexID ].ConnectedEdgeIDs.Num() == 0 )
 		{
-			InOutOrphanedVerticesPtr->Add( VertexID );
+			InOutOrphanedVerticesPtr->AddUnique( VertexID );
 		}
 		VertexInstanceArray.Remove( VertexInstanceID );
+		VertexInstanceAttributesSet.Remove( VertexInstanceID );
+	}
+
+	/** Reserves space for this number of new edges */
+	void ReserveNewEdges( const int32 NumEdges )
+	{
+		EdgeArray.Reserve( EdgeArray.Num() + NumEdges );
 	}
 
 private:
@@ -686,6 +388,7 @@ private:
 		Edge.ConnectedPolygons = ConnectedPolygons;
 		VertexArray[ VertexID0 ].ConnectedEdgeIDs.AddUnique( EdgeID );
 		VertexArray[ VertexID1 ].ConnectedEdgeIDs.AddUnique( EdgeID );
+		EdgeAttributesSet.Insert( EdgeID );
 	}
 
 public:
@@ -705,7 +408,7 @@ public:
 	}
 
 	/** Deletes an edge from a mesh */
-	void DeleteEdge( const FEdgeID EdgeID, TSet<FVertexID>* InOutOrphanedVerticesPtr = nullptr )
+	void DeleteEdge( const FEdgeID EdgeID, TArray<FVertexID>* InOutOrphanedVerticesPtr = nullptr )
 	{
 		FMeshEdge& Edge = EdgeArray[ EdgeID ];
 		for( const FVertexID EdgeVertexID : Edge.VertexIDs )
@@ -715,10 +418,17 @@ public:
 			if( InOutOrphanedVerticesPtr && Vertex.ConnectedEdgeIDs.Num() == 0 )
 			{
 				check( Vertex.VertexInstanceIDs.Num() == 0 );  // We must already have deleted any vertex instances
-				InOutOrphanedVerticesPtr->Add( EdgeVertexID );
+				InOutOrphanedVerticesPtr->AddUnique( EdgeVertexID );
 			}
 		}
 		EdgeArray.Remove( EdgeID );
+		EdgeAttributesSet.Remove( EdgeID );
+	}
+
+	/** Reserves space for this number of new polygons */
+	void ReserveNewPolygons( const int32 NumPolygons )
+	{
+		PolygonArray.Reserve( PolygonArray.Num() + NumPolygons );
 	}
 
 	/** Pair of IDs representing the vertex instance ID on a contour, and the edge ID which starts at that point, winding counter-clockwise */
@@ -746,7 +456,7 @@ private:
 		}
 	}
 
-	void CreatePolygon_Internal( const FPolygonID PolygonID, const TArray<FContourPoint>& Perimeter, const TArray<TArray<FContourPoint>>& Holes )
+	void CreatePolygon_Internal( const FPolygonID PolygonID, const FPolygonGroupID PolygonGroupID, const TArray<FContourPoint>& Perimeter, const TArray<TArray<FContourPoint>>& Holes )
 	{
 		FMeshPolygon& Polygon = PolygonArray[ PolygonID ];
 		CreatePolygonContour_Internal( PolygonID, Perimeter, Polygon.PerimeterContour );
@@ -755,26 +465,32 @@ private:
 			Polygon.HoleContours.Emplace();
 			CreatePolygonContour_Internal( PolygonID, Hole, Polygon.HoleContours.Last() );
 		}
+
+		Polygon.PolygonGroupID = PolygonGroupID;
+		check( !PolygonGroupArray[ PolygonGroupID ].Polygons.Contains( PolygonID ) );
+		PolygonGroupArray[ PolygonGroupID ].Polygons.Add( PolygonID );
+
+		PolygonAttributesSet.Insert( PolygonID );
 	}
 
 public:
 	/** Adds a new polygon to the mesh and returns its ID */
-	FPolygonID CreatePolygon( const TArray<FContourPoint>& Perimeter, const TArray<TArray<FContourPoint>>& Holes = TArray<TArray<FContourPoint>>() )
+	FPolygonID CreatePolygon( const FPolygonGroupID PolygonGroupID, const TArray<FContourPoint>& Perimeter, const TArray<TArray<FContourPoint>>& Holes = TArray<TArray<FContourPoint>>() )
 	{
 		const FPolygonID PolygonID = PolygonArray.Add();
-		CreatePolygon_Internal( PolygonID, Perimeter, Holes );
+		CreatePolygon_Internal( PolygonID, PolygonGroupID, Perimeter, Holes );
 		return PolygonID;
 	}
 
 	/** Adds a new polygon to the mesh with the given ID */
-	void CreatePolygonWithID( const FPolygonID PolygonID, const TArray<FContourPoint>& Perimeter, const TArray<TArray<FContourPoint>>& Holes = TArray<TArray<FContourPoint>>() )
+	void CreatePolygonWithID( const FPolygonID PolygonID, const FPolygonGroupID PolygonGroupID, const TArray<FContourPoint>& Perimeter, const TArray<TArray<FContourPoint>>& Holes = TArray<TArray<FContourPoint>>() )
 	{
 		PolygonArray.Insert( PolygonID );
-		CreatePolygon_Internal( PolygonID, Perimeter, Holes );
+		CreatePolygon_Internal( PolygonID, PolygonGroupID, Perimeter, Holes );
 	}
 
 private:
-	void DeletePolygonContour_Internal( const FPolygonID PolygonID, const TArray<FVertexInstanceID>& VertexInstanceIDs, TSet<FEdgeID>* InOutOrphanedEdgesPtr, TSet<FVertexInstanceID>* InOutOrphanedVertexInstancesPtr )
+	void DeletePolygonContour_Internal( const FPolygonID PolygonID, const TArray<FVertexInstanceID>& VertexInstanceIDs, TArray<FEdgeID>* InOutOrphanedEdgesPtr, TArray<FVertexInstanceID>* InOutOrphanedVertexInstancesPtr )
 	{
 		FVertexInstanceID LastVertexInstanceID = VertexInstanceIDs.Last();
 		for( const FVertexInstanceID VertexInstanceID : VertexInstanceIDs )
@@ -784,7 +500,7 @@ private:
 
 			if( InOutOrphanedVertexInstancesPtr && VertexInstance.ConnectedPolygons.Num() == 0 )
 			{
-				InOutOrphanedVertexInstancesPtr->Add( VertexInstanceID );
+				InOutOrphanedVertexInstancesPtr->AddUnique( VertexInstanceID );
 			}
 
 			const FEdgeID EdgeID = GetVertexPairEdge( VertexInstanceArray[ LastVertexInstanceID ].VertexID, VertexInstanceArray[ VertexInstanceID ].VertexID );
@@ -795,7 +511,7 @@ private:
 
 			if( InOutOrphanedEdgesPtr && Edge.ConnectedPolygons.Num() == 0 )
 			{
-				InOutOrphanedEdgesPtr->Add( EdgeID );
+				InOutOrphanedEdgesPtr->AddUnique( EdgeID );
 			}
 
 			LastVertexInstanceID = VertexInstanceID;
@@ -804,7 +520,7 @@ private:
 
 public:
 	/** Deletes a polygon from the mesh */
-	void DeletePolygon( const FPolygonID PolygonID, TSet<FEdgeID>* InOutOrphanedEdgesPtr = nullptr, TSet<FVertexInstanceID>* InOutOrphanedVertexInstancesPtr = nullptr, TSet<FPolygonGroupID>* InOutOrphanedPolygonGroupsPtr = nullptr )
+	void DeletePolygon( const FPolygonID PolygonID, TArray<FEdgeID>* InOutOrphanedEdgesPtr = nullptr, TArray<FVertexInstanceID>* InOutOrphanedVertexInstancesPtr = nullptr, TArray<FPolygonGroupID>* InOutOrphanedPolygonGroupsPtr = nullptr )
 	{
 		FMeshPolygon& Polygon = PolygonArray[ PolygonID ];
 		DeletePolygonContour_Internal( PolygonID, Polygon.PerimeterContour.VertexInstanceIDs, InOutOrphanedEdgesPtr, InOutOrphanedVertexInstancesPtr );
@@ -819,22 +535,39 @@ public:
 
 		if( InOutOrphanedPolygonGroupsPtr && PolygonGroup.Polygons.Num() == 0 )
 		{
-			InOutOrphanedPolygonGroupsPtr->Add( Polygon.PolygonGroupID );
+			InOutOrphanedPolygonGroupsPtr->AddUnique( Polygon.PolygonGroupID );
 		}
 
 		PolygonArray.Remove( PolygonID );
+		PolygonAttributesSet.Remove( PolygonID );
+	}
+
+	/** Reserves space for this number of new polygon groups */
+	void ReserveNewPolygonGroups( const int32 NumPolygonGroups )
+	{
+		PolygonGroupArray.Reserve( PolygonGroupArray.Num() + NumPolygonGroups );
+	}
+
+private:
+	void CreatePolygonGroup_Internal( const FPolygonGroupID PolygonGroupID )
+	{
+		PolygonGroupAttributesSet.Insert( PolygonGroupID );
 	}
 
+public:
 	/** Adds a new polygon group to the mesh and returns its ID */
 	FPolygonGroupID CreatePolygonGroup()
 	{
-		return PolygonGroupArray.Add();
+		const FPolygonGroupID PolygonGroupID = PolygonGroupArray.Add();
+		CreatePolygonGroup_Internal( PolygonGroupID );
+		return PolygonGroupID;
 	}
 
 	/** Adds a new polygon group to the mesh with the given ID */
 	void CreatePolygonGroupWithID( const FPolygonGroupID PolygonGroupID )
 	{
 		PolygonGroupArray.Insert( PolygonGroupID );
+		CreatePolygonGroup_Internal( PolygonGroupID );
 	}
 
 	/** Deletes a polygon group from the mesh */
@@ -842,6 +575,7 @@ public:
 	{
 		check( PolygonGroupArray[ PolygonGroupID ].Polygons.Num() == 0 );
 		PolygonGroupArray.Remove( PolygonGroupID );
+		PolygonGroupAttributesSet.Remove( PolygonGroupID );
 	}
 
 
@@ -861,6 +595,93 @@ public:
 		return FEdgeID::Invalid;
 	}
 
+	/** Returns reference to an array of Edge IDs connected to this vertex */
+	const TArray<FEdgeID>& GetVertexConnectedEdges( const FVertexID VertexID ) const
+	{
+		return VertexArray[ VertexID ].ConnectedEdgeIDs;
+	}
+
+	/** Returns reference to an array of VertexInstance IDs instanced from this vertex */
+	const TArray<FVertexInstanceID>& GetVertexVertexInstances( const FVertexID VertexID ) const
+	{
+		return VertexArray[ VertexID ].VertexInstanceIDs;
+	}
+
+	/** Returns reference to an array of polygon IDs connected to this edge */
+	const TArray<FPolygonID>& GetEdgeConnectedPolygons( const FEdgeID EdgeID ) const
+	{
+		return EdgeArray[ EdgeID ].ConnectedPolygons;
+	}
+
+	/** Returns a pair of vertex IDs defining the edge */
+	TTuple<FVertexID, FVertexID> GetEdgeVertices( const FEdgeID EdgeID ) const
+	{
+		const FMeshEdge& Edge = EdgeArray[ EdgeID ];
+		return MakeTuple( Edge.VertexIDs[ 0 ], Edge.VertexIDs[ 1 ] );
+	}
+
+	/** Returns reference to the array of triangles representing the triangulated polygon */
+	TArray<FMeshTriangle>& GetPolygonTriangles( const FPolygonID PolygonID )
+	{
+		return PolygonArray[ PolygonID ].Triangles;
+	}
+
+	const TArray<FMeshTriangle>& GetPolygonTriangles( const FPolygonID PolygonID ) const
+	{
+		return PolygonArray[ PolygonID ].Triangles;
+	}
+
+	/** Returns reference to an array of VertexInstance IDs forming the perimeter of this polygon */
+	const TArray<FVertexInstanceID>& GetPolygonPerimeterVertexInstances( const FPolygonID PolygonID ) const
+	{
+		return PolygonArray[ PolygonID ].PerimeterContour.VertexInstanceIDs;
+	}
+
+	/** Return the polygon group associated with a polygon */
+	FPolygonGroupID GetPolygonPolygonGroup( const FPolygonID PolygonID ) const
+	{
+		return PolygonArray[ PolygonID ].PolygonGroupID;
+	}
+
+	/** Sets the polygon group associated with a polygon */
+	void SetPolygonPolygonGroup( const FPolygonID PolygonID, const FPolygonGroupID PolygonGroupID )
+	{
+		FMeshPolygon& Polygon = PolygonArray[ PolygonID ];
+		verify( PolygonGroupArray[ Polygon.PolygonGroupID ].Polygons.Remove( PolygonID ) == 1 );
+		Polygon.PolygonGroupID = PolygonGroupID;
+		check( !PolygonGroupArray[ PolygonGroupID ].Polygons.Contains( PolygonID ) );
+		PolygonGroupArray[ PolygonGroupID ].Polygons.Add( PolygonID );
+	}
+
+	/** Return the vertex instance which corresponds to the given vertex on the given polygon, or FVertexInstanceID::Invalid */
+	FVertexInstanceID GetVertexInstanceForPolygonVertex( const FPolygonID PolygonID, const FVertexID VertexID ) const
+	{
+		const FVertexInstanceID* VertexInstanceIDPtr = PolygonArray[ PolygonID ].PerimeterContour.VertexInstanceIDs.FindByPredicate(
+			[ this, VertexID ]( const FVertexInstanceID VertexInstanceID )
+			{
+				return GetVertexInstanceVertex( VertexInstanceID ) == VertexID;
+			});
+
+		return VertexInstanceIDPtr ? *VertexInstanceIDPtr : FVertexInstanceID::Invalid;
+	}
+
+	/** Returns the vertex ID associated with the given vertex instance */
+	FVertexID GetVertexInstanceVertex( const FVertexInstanceID VertexInstanceID ) const
+	{
+		return VertexInstanceArray[ VertexInstanceID ].VertexID;
+	}
+
+	/** Returns reference to an array of Polygon IDs connected to this vertex instance */
+	const TArray<FPolygonID>& GetVertexInstanceConnectedPolygons( const FVertexInstanceID VertexInstanceID ) const
+	{
+		return VertexInstanceArray[ VertexInstanceID ].ConnectedPolygons;
+	}
+
+	/** Returns the polygons associated with the given polygon group */
+	const TArray<FPolygonID>& GetPolygonGroupPolygons( const FPolygonGroupID PolygonGroupID ) const
+	{
+		return PolygonGroupArray[ PolygonGroupID ].Polygons;
+	}
 
 	/** Compacts the data held in the mesh description, and returns an object describing how the IDs have been remapped. */
 	void Compact( FElementIDRemappings& OutRemappings );
@@ -874,6 +695,9 @@ private:
 	/** Given a set of index remappings, fixes up references to element IDs */
 	void FixUpElementIDs( const FElementIDRemappings& Remappings );
 
+	/** Given a set of index remappings, remaps all attributes accordingly */
+	void RemapAttributes( const FElementIDRemappings& Remappings );
+
 	FVertexArray VertexArray;
 	FVertexInstanceArray VertexInstanceArray;
 	FEdgeArray EdgeArray;
diff --git a/Engine/Source/Runtime/MeshDescription/Public/MeshElementArray.h b/Engine/Source/Runtime/MeshDescription/Public/MeshElementArray.h
--- a/Engine/Source/Runtime/MeshDescription/Public/MeshElementArray.h
+++ b/Engine/Source/Runtime/MeshDescription/Public/MeshElementArray.h
@@ -86,73 +86,73 @@ class TMeshElementArray : public TMeshElementArrayBase<ElementType>
 public:
 
 	/** Resets the container, optionally reserving space for elements to be added */
-	void Reset( const int32 Elements = 0 )
+	FORCEINLINE void Reset( const int32 Elements = 0 )
 	{
 		Container.Reset();
 		Container.Reserve( Elements );
 	}
 
 	/** Reserves space for the specified total number of elements */
-	void Reserve( const int32 Elements ) { Container.Reserve( Elements ); }
+	FORCEINLINE void Reserve( const int32 Elements ) { Container.Reserve( Elements ); }
 
 	/** Add a new element at the next available index, and return the new ID */
-	ElementIDType Add() { return ElementIDType( Container.Add( ElementType() ) ); }
+	FORCEINLINE ElementIDType Add() { return ElementIDType( Container.Add( ElementType() ) ); }
 
 	/** Add the provided element at the next available index, and return the new ID */
-	ElementIDType Add( typename TTypeTraits<ElementType>::ConstInitType Element ) { return ElementIDType( Container.Add( Element ) ); }
+	FORCEINLINE ElementIDType Add( typename TTypeTraits<ElementType>::ConstInitType Element ) { return ElementIDType( Container.Add( Element ) ); }
 
 	/** Add the provided element at the next available index, and return the ID */
-	ElementIDType Add( ElementType&& Element ) { return ElementIDType( Container.Add( Forward<ElementType>( Element ) ) ); }
+	FORCEINLINE ElementIDType Add( ElementType&& Element ) { return ElementIDType( Container.Add( Forward<ElementType>( Element ) ) ); }
 
 	/** Inserts a new element with the given ID */
-	ElementType& Insert( const ElementIDType ID )
+	FORCEINLINE ElementType& Insert( const ElementIDType ID )
 	{
 		Container.Insert( ID.GetValue(), ElementType() );
 		return Container[ ID.GetValue() ];
 	}
 
 	/** Inserts the provided element with the given ID */
-	ElementType& Insert( const ElementIDType ID, typename TTypeTraits<ElementType>::ConstInitType Element )
+	FORCEINLINE ElementType& Insert( const ElementIDType ID, typename TTypeTraits<ElementType>::ConstInitType Element )
 	{
 		Container.Insert( ID.GetValue(), Element );
 		return Container[ ID.GetValue() ];
 	}
 
 	/** Inserts the provided element with the given ID */
-	ElementType& Insert( const ElementIDType ID, ElementType&& Element )
+	FORCEINLINE ElementType& Insert( const ElementIDType ID, ElementType&& Element )
 	{
 		Container.Insert( ID.GetValue(), Forward<ElementType>( Element ) );
 		return Container[ ID.GetValue() ];
 	}
 
 	/** Removes the element with the given ID */
-	void Remove( const ElementIDType ID )
+	FORCEINLINE void Remove( const ElementIDType ID )
 	{
 		checkSlow( Container.IsAllocated( ID.GetValue() ) );
 		Container.RemoveAt( ID.GetValue() );
 	}
 
 	/** Returns the element with the given ID */
-	ElementType& operator[]( const ElementIDType ID )
+	FORCEINLINE ElementType& operator[]( const ElementIDType ID )
 	{
 		checkSlow( Container.IsAllocated( ID.GetValue() ) );
 		return Container[ ID.GetValue() ];
 	}
 
-	const ElementType& operator[]( const ElementIDType ID ) const
+	FORCEINLINE const ElementType& operator[]( const ElementIDType ID ) const
 	{
 		checkSlow( Container.IsAllocated( ID.GetValue() ) );
 		return Container[ ID.GetValue() ];
 	}
 
 	/** Returns the number of elements in the container */
-	int32 Num() const { return Container.Num(); }
+	FORCEINLINE int32 Num() const { return Container.Num(); }
 
 	/** Returns the index after the last valid element */
-	int32 GetArraySize() const { return Container.GetMaxIndex(); }
+	FORCEINLINE int32 GetArraySize() const { return Container.GetMaxIndex(); }
 
 	/** Returns the first valid ID */
-	ElementIDType GetFirstValidID() const
+	FORCEINLINE ElementIDType GetFirstValidID() const
 	{
 		return Container.Num() > 0 ?
 			ElementIDType( typename TSparseArray<ElementType>::TConstIterator( Container ).GetIndex() ) :
@@ -160,13 +160,13 @@ public:
 	}
 
 	/** Returns whether the given ID is valid or not */
-	bool IsValid( const ElementIDType ID ) const
+	FORCEINLINE bool IsValid( const ElementIDType ID ) const
 	{
 		return ID.GetValue() >= 0 && ID.GetValue() < Container.GetMaxIndex() && Container.IsAllocated( ID.GetValue() );
 	}
 
 	/** Serializer */
-	friend FArchive& operator<<( FArchive& Ar, TMeshElementArray& Array )
+	FORCEINLINE friend FArchive& operator<<( FArchive& Ar, TMeshElementArray& Array )
 	{
 		Ar << static_cast<TMeshElementArrayBase<ElementType>&>( Array );
 		return Ar;
@@ -192,7 +192,7 @@ public:
 	{
 	public:
 
-		explicit TElementIDs( const TSparseArray<ElementType>& InArray )
+		explicit FORCEINLINE TElementIDs( const TSparseArray<ElementType>& InArray )
 			: Array( InArray )
 		{}
 
@@ -200,27 +200,27 @@ public:
 		{
 		public:
 
-			explicit TConstIterator( typename TSparseArray<ElementType>::TConstIterator&& It )
+			explicit FORCEINLINE TConstIterator( typename TSparseArray<ElementType>::TConstIterator&& It )
 				: Iterator( MoveTemp( It ) )
 			{}
 
-			TConstIterator& operator++()
+			FORCEINLINE TConstIterator& operator++()
 			{
 				++Iterator;
 				return *this;
 			}
 
-			ElementIDType operator*() const
+			FORCEINLINE ElementIDType operator*() const
 			{
 				return Iterator ? ElementIDType( Iterator.GetIndex() ) : ElementIDType::Invalid;
 			}
 
-			friend bool operator==( const TConstIterator& Lhs, const TConstIterator& Rhs )
+			friend FORCEINLINE bool operator==( const TConstIterator& Lhs, const TConstIterator& Rhs )
 			{
 				return Lhs.Iterator == Rhs.Iterator;
 			}
 
-			friend bool operator!=( const TConstIterator& Lhs, const TConstIterator& Rhs )
+			friend FORCEINLINE bool operator!=( const TConstIterator& Lhs, const TConstIterator& Rhs )
 			{
 				return Lhs.Iterator != Rhs.Iterator;
 			}
@@ -230,19 +230,19 @@ public:
 			typename TSparseArray<ElementType>::TConstIterator Iterator;
 		};
 
-		TConstIterator CreateConstIterator() const
+		FORCEINLINE TConstIterator CreateConstIterator() const
 		{
 			return TConstIterator( typename TSparseArray<ElementType>::TConstIterator( Array ) );
 		}
 
 	private:
 
-		friend TConstIterator begin( const TElementIDs& ArrayIDs )
+		friend FORCEINLINE TConstIterator begin( const TElementIDs& ArrayIDs )
 		{
 			return TConstIterator( begin( ArrayIDs.Array ) );
 		}
 
-		friend TConstIterator end( const TElementIDs& ArrayIDs )
+		friend FORCEINLINE TConstIterator end( const TElementIDs& ArrayIDs )
 		{
 			return TConstIterator( end( ArrayIDs.Array ) );
 		}
@@ -251,7 +251,7 @@ public:
 	};
 
 	/** Return iterable proxy object from container */
-	TElementIDs GetElementIDs() const { return TElementIDs( Container ); }
+	TElementIDs FORCEINLINE GetElementIDs() const { return TElementIDs( Container ); }
 };
 
 
diff --git a/Engine/Source/Runtime/MeshDescription/Public/MeshElementRemappings.h b/Engine/Source/Runtime/MeshDescription/Public/MeshElementRemappings.h
new file mode 100644
--- /dev/null
+++ b/Engine/Source/Runtime/MeshDescription/Public/MeshElementRemappings.h
@@ -0,0 +1,49 @@
+// Copyright 1998-2017 Epic Games, Inc. All Rights Reserved.
+
+#pragma once
+
+#include "CoreMinimal.h"
+#include "MeshTypes.h"
+
+/**
+ * This is a structure which holds the ID remappings returned by a Compact operation, or passed to a Remap operation.
+ */
+struct FElementIDRemappings
+{
+	TSparseArray<FVertexID> NewVertexIndexLookup;
+	TSparseArray<FVertexInstanceID> NewVertexInstanceIndexLookup;
+	TSparseArray<FEdgeID> NewEdgeIndexLookup;
+	TSparseArray<FPolygonID> NewPolygonIndexLookup;
+	TSparseArray<FPolygonGroupID> NewPolygonGroupIndexLookup;
+
+	FVertexID GetRemappedVertexID( FVertexID VertexID ) const
+	{
+		checkSlow( NewVertexIndexLookup.IsAllocated( VertexID.GetValue() ) );
+		return NewVertexIndexLookup[ VertexID.GetValue() ];
+	}
+
+	FVertexInstanceID GetRemappedVertexInstanceID( FVertexInstanceID VertexInstanceID ) const
+	{
+		checkSlow( NewVertexInstanceIndexLookup.IsAllocated( VertexInstanceID.GetValue() ) );
+		return NewVertexInstanceIndexLookup[ VertexInstanceID.GetValue() ];
+	}
+
+	FEdgeID GetRemappedEdgeID( FEdgeID EdgeID ) const
+	{
+		checkSlow( NewEdgeIndexLookup.IsAllocated( EdgeID.GetValue() ) );
+		return NewEdgeIndexLookup[ EdgeID.GetValue() ];
+	}
+
+	FPolygonID GetRemappedPolygonID( FPolygonID PolygonID ) const
+	{
+		checkSlow( NewPolygonIndexLookup.IsAllocated( PolygonID.GetValue() ) );
+		return NewPolygonIndexLookup[ PolygonID.GetValue() ];
+	}
+
+	FPolygonGroupID GetRemappedPolygonGroupID( FPolygonGroupID PolygonGroupID ) const
+	{
+		check( NewPolygonGroupIndexLookup.IsAllocated( PolygonGroupID.GetValue() ) );
+		return NewPolygonGroupIndexLookup[ PolygonGroupID.GetValue() ];
+	}
+};
+
diff --git a/Engine/Source/Runtime/MeshDescription/Public/MeshTypes.h b/Engine/Source/Runtime/MeshDescription/Public/MeshTypes.h
--- a/Engine/Source/Runtime/MeshDescription/Public/MeshTypes.h
+++ b/Engine/Source/Runtime/MeshDescription/Public/MeshTypes.h
@@ -204,6 +204,7 @@ struct FPolygonID : public FElementID
 };
 
 
+#if 0
 UCLASS( abstract )
 class MESHDESCRIPTION_API UEditableMeshAttribute : public UObject
 {
@@ -307,3 +308,4 @@ private:
 	static const FName PolygonNormalName;
 	static const FName PolygonCenterName;
 };
+#endif
