Disallowed use of the TMeshAttributeArray copy constructor so that Attribute Arrays are not mistakenly accessed by value (they must always be assigned to a [const] TMeshAttributeArray<>&).

diff --git a/Engine/Source/Runtime/MeshDescription/Public/MeshAttributeArray.h b/Engine/Source/Runtime/MeshDescription/Public/MeshAttributeArray.h
--- a/Engine/Source/Runtime/MeshDescription/Public/MeshAttributeArray.h
+++ b/Engine/Source/Runtime/MeshDescription/Public/MeshAttributeArray.h
@@ -33,6 +33,12 @@ class TMeshAttributeArrayBase
 {
 public:
 
+	/** Disallow use of the copy constructor to prevent arrays being mistakenly accessed in the UMeshDescription by value */
+	TMeshAttributeArrayBase( const TMeshAttributeArrayBase& )
+	{
+		static_assert( false, "TMeshAttributeArray must be accessed by reference");
+	}
+
 	/**
 	 * Custom serialization for TMeshAttributeArrayBase.
 	 */
@@ -44,6 +50,9 @@ public:
 
 protected:
 
+	/** Should not instance this base class directly */
+	TMeshAttributeArrayBase() = default;
+
 	/** Expands the array if necessary so that the passed element index is valid. Newly created elements will be assigned the default value. */
 	void Insert( const int32 Index, const ElementType& Default )
 	{
@@ -76,7 +85,7 @@ class TAttributeIndicesArray;
  * This derived class imposes this type safety.
  */
 template <typename ElementType, typename ElementIDType>
-class TMeshAttributeArray : public TMeshAttributeArrayBase<ElementType>
+class TMeshAttributeArray : private TMeshAttributeArrayBase<ElementType>
 {
 	static_assert( TIsDerivedFrom<ElementIDType, FElementID>::IsDerived, "ElementIDType must be derived from FElementID" );
 
@@ -167,8 +176,6 @@ class TAttributeIndicesArray
 {
 public:
 
-	using Type = AttributeType;
-
 	/** Default constructor - required so that it builds correctly */
 	TAttributeIndicesArray() = default;
 
@@ -319,6 +326,7 @@ template <typename ElementIDType>
 class TAttributesSet
 {
 public:
+
 	/**
 	 * Register a new attribute name with the given type (must be a member of the AttributeTypes tuple).
 	 *
