Changed TAttributesSet serialization so that it respects the currently registered attributes. If the loaded data is missing attributes which have been registered, they will be default initialized. If it contains attributes which are no longer registered, they will be deleted.

diff --git a/Engine/Source/Runtime/MeshDescription/Public/MeshAttributeArray.h b/Engine/Source/Runtime/MeshDescription/Public/MeshAttributeArray.h
--- a/Engine/Source/Runtime/MeshDescription/Public/MeshAttributeArray.h
+++ b/Engine/Source/Runtime/MeshDescription/Public/MeshAttributeArray.h
@@ -171,11 +171,14 @@ ENUM_CLASS_FLAGS( EMeshAttributeFlags );
  * This class represents a container for a named attribute on a mesh element.
  * It contains an array of TMeshAttributeArrays, one per attribute index. 
  */
-template <typename AttributeType, typename ElementIDType>
+template <typename T, typename U>
 class TAttributeIndicesArray
 {
 public:
 
+	using AttributeType = T;
+	using ElementIDType = U;
+
 	/** Default constructor - required so that it builds correctly */
 	TAttributeIndicesArray() = default;
 
@@ -184,6 +187,7 @@ public:
 		: DefaultValue( InDefaultValue ),
 		  Flags( InFlags )
 	{
+		check( NumberOfIndices > 0 );
 		ArrayForIndices.SetNum( NumberOfIndices );
 	}
 
@@ -212,18 +216,22 @@ public:
 	/** Return flags for this attribute type */
 	FORCEINLINE EMeshAttributeFlags GetFlags() const { return Flags; }
 
+	/** Return default value for this attribute type */
+	FORCEINLINE AttributeType GetDefaultValue() const { return DefaultValue; }
+
 	/** Return number of indices this attribute has */
 	FORCEINLINE int32 GetNumIndices() const { return ArrayForIndices.Num(); }
 
 	/** Sets number of indices this attribute has */
 	void SetNumIndices( const int32 NumIndices )
 	{
+		check( NumIndices > 0 );
 		const int32 OriginalNumIndices = ArrayForIndices.Num();
 		ArrayForIndices.SetNum( NumIndices );
 
-		// If there is already at least one attribute index, and it is non-empty, ensure that newly added indices
+		// If the first attribute index is non-empty, ensure that newly added indices
 		// are initialized to the same size with the default value (they must all have equal size).
-		if( OriginalNumIndices > 0 && ArrayForIndices[ 0 ].Num() > 0 )
+		if( ArrayForIndices[ 0 ].Num() > 0 )
 		{
 			for( int32 Index = OriginalNumIndices; Index < NumIndices; ++Index )
 			{
@@ -340,10 +348,21 @@ public:
 	template <typename AttributeType>
 	void RegisterAttribute( const FName AttributeName, const int32 NumberOfIndices = 1, const AttributeType& Default = AttributeType(), const EMeshAttributeFlags Flags = EMeshAttributeFlags::None )
 	{
-		Container.Get<TTupleIndex<AttributeType, AttributeTypes>::Value>().Emplace(
-			AttributeName,
-			TAttributeIndicesArray<AttributeType, ElementIDType>( NumberOfIndices, Default, Flags )
-		);
+		// See if there is an existing attribute in this set, and get the number of elements it has.
+		// When a new attribute is registered in a non-empty attribute set, it needs to be default-initialized to contain the same
+		// number of elements as its contemporaries.
+		TAttributesMap<AttributeType, ElementIDType>& AttributesMap = Container.Get<TTupleIndex<AttributeType, AttributeTypes>::Value>();
+		typename TAttributesMap<AttributeType, ElementIDType>::TConstIterator It( AttributesMap.CreateConstIterator() );
+		const int32 NumElements = ( !It ) ? 0 : It.Value().GetArrayForIndex( 0 ).Num();
+
+		if( !AttributesMap.Contains( AttributeName ) )
+		{
+			TAttributeIndicesArray<AttributeType, ElementIDType>& AttributeIndicesArray = AttributesMap.Emplace( AttributeName, TAttributeIndicesArray<AttributeType, ElementIDType>( NumberOfIndices, Default, Flags ) );
+			if( NumElements > 0 )
+			{
+				AttributeIndicesArray.Insert( ElementIDType( NumElements - 1 ) );
+			}
+		}
 	}
 
 	/** Unregister an attribute name with the given type */
@@ -524,7 +543,87 @@ public:
 	/** Serializer */
 	friend FArchive& operator<<( FArchive& Ar, TAttributesSet& AttributesSet )
 	{
-		Ar << AttributesSet.Container;
+		auto SerializeElements = [ &Ar ]( auto& AttributesMap )
+		{
+			using AttributeIndicesArrayType = decltype( DeclVal<typename TDecay<decltype( AttributesMap )>::Type::ElementType>().Value );
+			using AttributeType = typename AttributeIndicesArrayType::AttributeType;
+
+			struct FRegisteredAttribute
+			{
+				FName Name;
+				int32 NumIndices;
+				AttributeType DefaultValue;
+				EMeshAttributeFlags Flags;
+			};
+
+			// If we are loading, get a list of currently registered attributes and their metadata.
+			// This will be used to update the loaded attribute map after serialization.
+			TArray<FRegisteredAttribute> RegisteredAttributes;
+			if( Ar.IsLoading() )
+			{
+				for( const auto& AttributeNameAndIndicesArray : AttributesMap )
+				{
+					RegisteredAttributes.Emplace();
+					FRegisteredAttribute& RegisteredAttribute = RegisteredAttributes.Last();
+					RegisteredAttribute.Name = AttributeNameAndIndicesArray.Key;
+					RegisteredAttribute.NumIndices = AttributeNameAndIndicesArray.Value.GetNumIndices();
+					RegisteredAttribute.DefaultValue = AttributeNameAndIndicesArray.Value.GetDefaultValue();
+					RegisteredAttribute.Flags = AttributeNameAndIndicesArray.Value.GetFlags();
+				}
+			}
+
+			// Now serialize the attributes of this type.
+			Ar << AttributesMap;
+
+			// If we are loading, the attributes just serialized may be out of sync with the currently registered attributes, so we need to fix this up now.
+			if( Ar.IsLoading() )
+			{
+				// Determine the number of elements for the attributes just serialized
+				auto It( AttributesMap.CreateConstIterator() );
+				const int32 NumElements = ( !It ) ? 0 : It.Value().GetArrayForIndex( 0 ).Num();
+
+				// Go through the list of registered attributes we just built
+				for( const FRegisteredAttribute& RegisteredAttribute : RegisteredAttributes )
+				{
+					AttributeIndicesArrayType* AttributeIndicesArray = AttributesMap.Find( RegisteredAttribute.Name );
+					if( !AttributeIndicesArray )
+					{
+						// Re-register any attributes that aren't there
+						AttributeIndicesArrayType NewAttributeIndicesArray( RegisteredAttribute.NumIndices, RegisteredAttribute.DefaultValue, RegisteredAttribute.Flags );
+						if( NumElements > 0 )
+						{
+							// and fill the new attribute with default initialized values
+							NewAttributeIndicesArray.Insert( ElementIDType( NumElements - 1 ) );
+						}
+
+						UE_LOG( LogTemp, Warning, TEXT( "Didn't find attribute '%s' - adding" ), *RegisteredAttribute.Name.ToString() );
+						AttributesMap.Add( RegisteredAttribute.Name, MoveTemp( NewAttributeIndicesArray ) );
+					}
+					else if( AttributeIndicesArray->GetNumIndices() != RegisteredAttribute.NumIndices )
+					{
+						// Amend the number of attribute indices for any which are not correct
+						UE_LOG( LogTemp, Warning, TEXT( "Found attribute '%s' with the wrong number of indices - amending" ), *RegisteredAttribute.Name.ToString() );
+						AttributeIndicesArray->SetNumIndices( RegisteredAttribute.NumIndices );
+					}
+				}
+
+				// Finally, go through all the attributes in the container, and remove any entries which are not in the registered list
+				for( auto AttributesIt = AttributesMap.CreateIterator(); AttributesIt; ++AttributesIt )
+				{
+					const FName& AttributeName = AttributesIt.Key();
+					if( !RegisteredAttributes.ContainsByPredicate(
+						[ &AttributeName ]( const FRegisteredAttribute& RegisteredAttribute ) { return RegisteredAttribute.Name == AttributeName; }
+					  ) )
+					{
+						UE_LOG( LogTemp, Warning, TEXT( "Found unregistered attribute '%s' - removing" ), *AttributeName.ToString() );
+						AttributesIt.RemoveCurrent();
+					}
+				}
+			}
+		};
+
+		VisitTupleElements( AttributesSet.Container, SerializeElements );
+
 		return Ar;
 	}
 
