Default texture coordinate number of indices to 1 instead of 2 (this should always be reinitialized to the correct value when initializing a mesh description from an old static mesh). Eventually we should not be registering any attributes in the UMeshDescription constructor; this should be done in the class which owns the mesh description.

diff --git a/Engine/Source/Runtime/MeshDescription/Private/MeshDescription.cpp b/Engine/Source/Runtime/MeshDescription/Private/MeshDescription.cpp
--- a/Engine/Source/Runtime/MeshDescription/Private/MeshDescription.cpp
+++ b/Engine/Source/Runtime/MeshDescription/Private/MeshDescription.cpp
@@ -18,7 +18,7 @@ void UMeshDescription::RegisterBaseAttributes()
 	VertexAttributes().RegisterAttribute<float>(MeshAttribute::Vertex::CornerSharpness, 1, 0.0f, EMeshAttributeFlags::Lerpable);
 
 	// Add basic vertex instance attributes
-	VertexInstanceAttributes().RegisterAttribute<FVector2D>(MeshAttribute::VertexInstance::TextureCoordinate, 2, FVector2D::ZeroVector, EMeshAttributeFlags::Lerpable);
+	VertexInstanceAttributes().RegisterAttribute<FVector2D>(MeshAttribute::VertexInstance::TextureCoordinate, 1, FVector2D::ZeroVector, EMeshAttributeFlags::Lerpable);
 	VertexInstanceAttributes().RegisterAttribute<FVector>(MeshAttribute::VertexInstance::Normal, 1, FVector::ZeroVector, EMeshAttributeFlags::AutoGenerated);
 	VertexInstanceAttributes().RegisterAttribute<FVector>(MeshAttribute::VertexInstance::Tangent, 1, FVector::ZeroVector, EMeshAttributeFlags::AutoGenerated);
 	VertexInstanceAttributes().RegisterAttribute<float>(MeshAttribute::VertexInstance::BinormalSign, 1, 0.0f, EMeshAttributeFlags::AutoGenerated);
