Revamp to mesh element attribute array:
- Serialization now works if an asset is deserialized with fewer attribute types than the current count.
- Fixed an issue in Remap where empty attribute maps wouldn't have their number of elements adjusted in line with other attribute maps.
- Added a way of initializing the attribute array to a fixed element size (to default attribute values).
- Added a new TAttributesMap class which acts as a wrapper for the actual TMap.
- Bumped the StaticMesh DDC GUID.

diff --git a/Engine/Source/Runtime/Engine/Private/StaticMesh.cpp b/Engine/Source/Runtime/Engine/Private/StaticMesh.cpp
--- a/Engine/Source/Runtime/Engine/Private/StaticMesh.cpp
+++ b/Engine/Source/Runtime/Engine/Private/StaticMesh.cpp
@@ -2370,7 +2370,7 @@ void UStaticMesh::FixupMaterialSlotName()
 // differences, etc.) replace the version GUID below with a new one.
 // In case of merge conflicts with DDC versions, you MUST generate a new GUID
 // and set this new GUID as the version.                                       
-#define MESHDATAKEY_STATICMESH_DERIVEDDATA_VER TEXT("81EFE7970F3140CABD7B24620A2BE2BE")
+#define MESHDATAKEY_STATICMESH_DERIVEDDATA_VER TEXT("3F6401D9624A4E2AB07998C997BC52D5")
 
 static const FString& GetMeshDataKeyStaticMeshDerivedDataVersion()
 {
diff --git a/Engine/Source/Runtime/MeshDescription/Private/MeshDescription.cpp b/Engine/Source/Runtime/MeshDescription/Private/MeshDescription.cpp
--- a/Engine/Source/Runtime/MeshDescription/Private/MeshDescription.cpp
+++ b/Engine/Source/Runtime/MeshDescription/Private/MeshDescription.cpp
@@ -6,6 +6,10 @@
 #include "MeshAttributes.h"
 #include "Serialization/ObjectWriter.h"
 
+
+DEFINE_LOG_CATEGORY( LogMeshDescription );
+
+
 UMeshDescription::UMeshDescription()
 {
 	RegisterBaseAttributes();
@@ -98,8 +102,11 @@ void UMeshDescription::Empty()
 	PolygonGroupArray.Reset();
 
 	//Empty all attributes
-	UnRegisterBaseAttributes();
-	RegisterBaseAttributes();
+	VertexAttributesSet.Initialize( 0 );
+	VertexInstanceAttributesSet.Initialize( 0 );
+	EdgeAttributesSet.Initialize( 0 );
+	PolygonAttributesSet.Initialize( 0 );
+	PolygonGroupAttributesSet.Initialize( 0 );
 }
 
 #if WITH_EDITORONLY_DATA
diff --git a/Engine/Source/Runtime/MeshDescription/Public/MeshAttributeArray.h b/Engine/Source/Runtime/MeshDescription/Public/MeshAttributeArray.h
--- a/Engine/Source/Runtime/MeshDescription/Public/MeshAttributeArray.h
+++ b/Engine/Source/Runtime/MeshDescription/Public/MeshAttributeArray.h
@@ -6,10 +6,15 @@
 #include "MeshTypes.h"
 #include "MeshElementRemappings.h"
 
+
 /**
  * List of attribute types which are supported.
  * We do this so we can automatically generate the attribute containers and their associated accessors with
- * some template magic. Adding a new attribute type requires no extra code elsewhere in the class.
+ * some template magic.
+ *
+ * IMPORTANT NOTE: Do not remove any type from this tuple, or serialization will fail.
+ * Types may be added at the end of this list if necessary, although please do so sparingly as each extra type will
+ * impact on performance and object size.
  */
 using AttributeTypes = TTuple
 <
@@ -26,6 +31,9 @@ using AttributeTypes = TTuple
 /**
  * This defines the container used to hold mesh element attributes of a particular name and index.
  * It is a simple TArray, so that all attributes are packed contiguously for each element ID.
+ *
+ * Note that the container may grow arbitrarily as new elements are inserted, but it will never be
+ * shrunk as elements are removed. The only operations that will shrink the container are Initialize() and Remap().
  */
 template <typename ElementType>
 class TMeshAttributeArrayBase
@@ -71,6 +79,13 @@ protected:
 		}
 	}
 
+	/** Initializes the array to the given size with the default value */
+	void Initialize( const int32 ElementCount, const ElementType& Default )
+	{
+		Container.Reset( ElementCount );
+		Insert( ElementCount - 1, Default );
+	}
+
 	/** The actual container, represented by a regular array */
 	TArray<ElementType> Container;
 };
@@ -182,12 +197,12 @@ public:
 	TAttributeIndicesArray() = default;
 
 	/** Constructor */
-	TAttributeIndicesArray( const int32 NumberOfIndices, const AttributeType& InDefaultValue, const EMeshAttributeFlags InFlags )
-		: DefaultValue( InDefaultValue ),
+	TAttributeIndicesArray( const int32 NumberOfIndices, const AttributeType& InDefaultValue, const EMeshAttributeFlags InFlags, const int32 InNumberOfElements )
+		: NumElements( InNumberOfElements ),
+		  DefaultValue( InDefaultValue ),
 		  Flags( InFlags )
 	{
-		check( NumberOfIndices > 0 );
-		ArrayForIndices.SetNum( NumberOfIndices );
+		SetNumIndices( NumberOfIndices );
 	}
 
 	/** Insert the element at the given index */
@@ -197,6 +212,8 @@ public:
 		{
 			ArrayForIndex.Insert( ElementID, DefaultValue );
 		}
+
+		NumElements = FMath::Max( NumElements, ElementID.GetValue() + 1 );
 	}
 
 	/** Remove the element at the given index, replacing it with a default value */
@@ -221,6 +238,9 @@ public:
 	/** Return number of indices this attribute has */
 	FORCEINLINE int32 GetNumIndices() const { return ArrayForIndices.Num(); }
 
+	/** Return number of elements each attribute index has */
+	FORCEINLINE int32 GetNumElements() const { return NumElements; }
+
 	/** Sets number of indices this attribute has */
 	void SetNumIndices( const int32 NumIndices )
 	{
@@ -228,14 +248,20 @@ public:
 		const int32 OriginalNumIndices = ArrayForIndices.Num();
 		ArrayForIndices.SetNum( NumIndices );
 
-		// If the first attribute index is non-empty, ensure that newly added indices
-		// are initialized to the same size with the default value (they must all have equal size).
-		if( ArrayForIndices[ 0 ].Num() > 0 )
+		// If we have added new indices, ensure they are filled out with the correct number of elements
+		for( int32 Index = OriginalNumIndices; Index < NumIndices; ++Index )
 		{
-			for( int32 Index = OriginalNumIndices; Index < NumIndices; ++Index )
-			{
-				ArrayForIndices[ Index ].Insert( ElementIDType( ArrayForIndices[ 0 ].Num() - 1 ), DefaultValue );
-			}
+			ArrayForIndices[ Index ].Initialize( NumElements, DefaultValue );
+		}
+	}
+
+	/** Sets the number of elements to the exact number provided, and initializes them to the default value */
+	void Initialize( const int32 Count )
+	{
+		NumElements = Count;
+		for( TMeshAttributeArray<AttributeType, ElementIDType>& ArrayForIndex : ArrayForIndices )
+		{
+			ArrayForIndex.Initialize( Count, DefaultValue );
 		}
 	}
 
@@ -245,12 +271,14 @@ public:
 		for( TMeshAttributeArray<AttributeType, ElementIDType>& ArrayForIndex : ArrayForIndices )
 		{
 			ArrayForIndex.Remap( IndexRemap, DefaultValue );
+			NumElements = ArrayForIndex.Num();
 		}
 	}
 
 	/** Serializer */
 	friend FArchive& operator<<( FArchive& Ar, TAttributeIndicesArray& AttributesArray )
 	{
+		Ar << AttributesArray.NumElements;
 		Ar << AttributesArray.ArrayForIndices;
 		Ar << AttributesArray.DefaultValue;
 		Ar << AttributesArray.Flags;
@@ -258,6 +286,9 @@ public:
 	}
 
 private:
+	/** Number of elements in each index */
+	int32 NumElements;
+
 	/** An array of MeshAttributeArrays, one per attribute index */
 	TArray<TMeshAttributeArray<AttributeType, ElementIDType>> ArrayForIndices;
 
@@ -278,10 +309,266 @@ template <typename AttributeType> using TPolygonGroupAttributeIndicesArray = TAt
 
 
 /**
- * This alias maps an attribute name to a TAttributeIndicesArray, i.e. an array of MeshAttributeArrays, one per attribute index.
+ * This maps an attribute name to a TAttributeIndicesArray, i.e. an array of MeshAttributeArrays, one per attribute index.
  */
-template <typename AttributeType, typename ElementIDType>
-using TAttributesMap = TMap<FName, TAttributeIndicesArray<AttributeType, ElementIDType>>;
+template <typename Attribute, typename ElementID>
+class TAttributesMap
+{
+public:
+	using AttributeType = Attribute;
+	using ElementIDType = ElementID;
+	using AttributeIndicesArrayType = TAttributeIndicesArray<AttributeType, ElementIDType>;
+	using MapType = TMap<FName, AttributeIndicesArrayType>;
+
+	FORCEINLINE TAttributesMap()
+		: NumElements( 0 )
+	{}
+
+	/** Register an attribute name */
+	FORCEINLINE void RegisterAttribute( const FName AttributeName, const int32 NumberOfIndices, const AttributeType& Default, const EMeshAttributeFlags Flags )
+	{
+		if( !Map.Contains( AttributeName ) )
+		{
+			Map.Emplace( AttributeName, TAttributeIndicesArray<AttributeType, ElementIDType>( NumberOfIndices, Default, Flags, NumElements ) );
+		}
+	}
+
+	/** Unregister an attribute name */
+	FORCEINLINE void UnregisterAttribute( const FName AttributeName )
+	{
+		Map.Remove( AttributeName );
+	}
+
+	/** Determines whether an attribute exists with the given name */
+	FORCEINLINE bool HasAttribute( const FName AttributeName ) const
+	{
+		return Map.Contains( AttributeName );
+	}
+
+	/** Get attribute array with the given name and index */
+	FORCEINLINE TMeshAttributeArray<AttributeType, ElementIDType>& GetAttributes( const FName AttributeName, const int32 AttributeIndex = 0 )
+	{
+		// @todo mesh description: should this handle non-existent attribute names and indices gracefully?
+		return Map.FindChecked( AttributeName ).GetArrayForIndex( AttributeIndex );
+	}
+
+	FORCEINLINE const TMeshAttributeArray<AttributeType, ElementIDType>& GetAttributes( const FName AttributeName, const int32 AttributeIndex = 0 ) const
+	{
+		// @todo mesh description: should this handle non-existent attribute names and indices gracefully?
+		return Map.FindChecked( AttributeName ).GetArrayForIndex( AttributeIndex );
+	}
+
+	/** Get attribute indices array with the given name */
+	FORCEINLINE TAttributeIndicesArray<AttributeType, ElementIDType>& GetAttributesSet( const FName AttributeName )
+	{
+		// @todo mesh description: should this handle non-existent attribute names gracefully?
+		return Map.FindChecked( AttributeName );
+	}
+
+	FORCEINLINE const TAttributeIndicesArray<AttributeType, ElementIDType>& GetAttributesSet( const FName AttributeName ) const
+	{
+		// @todo mesh description: should this handle non-existent attribute names gracefully?
+		return Map.FindChecked( AttributeName );
+	}
+
+	/** Returns the number of indices for the attribute with the given name */
+	FORCEINLINE int32 GetAttributeIndexCount( const FName AttributeName ) const
+	{
+		// @todo mesh description: should this handle non-existent attribute names and indices gracefully?
+		return Map.FindChecked( AttributeName ).GetNumIndices();
+	}
+
+	/** Sets the number of indices for the attribute with the given name */
+	FORCEINLINE void SetAttributeIndexCount( const FName AttributeName, const int32 NumIndices )
+	{
+		Map.FindChecked( AttributeName ).SetNumIndices( NumIndices );
+	}
+
+	/** Returns an array of all the attribute names registered for this attribute type */
+	template <typename Allocator>
+	FORCEINLINE void GetAttributeNames( TArray<FName, Allocator>& OutAttributeNames ) const
+	{
+		Map.GetKeys( OutAttributeNames );
+	}
+
+	/** Gets a single attribute with the given ElementID, Name and Index */
+	FORCEINLINE AttributeType GetAttribute( const ElementIDType ElementID, const FName AttributeName, const int32 AttributeIndex = 0 ) const
+	{
+		return Map.FindChecked( AttributeName ).GetArrayForIndex( AttributeIndex )[ ElementID ];
+	}
+
+	/** Sets a single attribute with the given ElementID, Name and Index to the given value */
+	FORCEINLINE void SetAttribute( const ElementIDType ElementID, const FName AttributeName, const int32 AttributeIndex, const AttributeType& AttributeValue )
+	{
+		Map.FindChecked( AttributeName ).GetArrayForIndex( AttributeIndex )[ ElementID ] = AttributeValue;
+	}
+
+	/** Inserts a default-initialized value for all attributes of the given ID */
+	void Insert( const ElementIDType ElementID )
+	{
+		NumElements = FMath::Max( NumElements, ElementID.GetValue() + 1 );
+		for( auto& AttributeNameAndIndicesArray : Map )
+		{
+			AttributeNameAndIndicesArray.Value.Insert( ElementID );
+			check( AttributeNameAndIndicesArray.Value.GetNumElements() == this->NumElements );
+		}
+	}
+
+	/** Removes all attributes with the given ID */
+	void Remove( const ElementIDType ElementID )
+	{
+		for( auto& AttributeNameAndIndicesArray : Map )
+		{
+			AttributeNameAndIndicesArray.Value.Remove( ElementID );
+		}
+	}
+
+	/** Initializes all attributes to have the given number of elements with the default value */
+	void Initialize( const int32 Count )
+	{
+		NumElements = Count;
+		for( auto& AttributeNameAndIndicesArray : Map )
+		{
+			AttributeNameAndIndicesArray.Value.Initialize( Count );
+		}
+	}
+
+	/** Returns the number of elements held by each attribute in this map */
+	FORCEINLINE int32 GetNumElements() const
+	{
+		return NumElements;
+	}
+
+	/**
+	 * Call the supplied function on each attribute.
+	 * The prototype should be Func( const FName AttributeName, auto& AttributeIndicesArray );
+	 */
+	template <typename FuncType>
+	void ForEachAttributeIndicesArray( const FuncType& Func )
+	{
+		for( auto& AttributeNameAndIndicesArray : Map )
+		{
+			Func( AttributeNameAndIndicesArray.Key, AttributeNameAndIndicesArray.Value );
+		}
+	}
+
+	template <typename FuncType>
+	void ForEachAttributeIndicesArray( const FuncType& Func ) const
+	{
+		for( const auto& AttributeNameAndIndicesArray : Map )
+		{
+			Func( AttributeNameAndIndicesArray.Key, AttributeNameAndIndicesArray.Value );
+		}
+	}
+
+	/** Applies the given remapping to the attributes set */
+	void Remap( const TSparseArray<ElementIDType>& IndexRemap )
+	{
+		if( Map.Num() == 0 )
+		{
+			// If there are no attributes registered, determine the number of elements by finding the maximum
+			// remapped element ID in the IndexRemap array
+			NumElements = 0;
+			for( const ElementIDType ElementID : IndexRemap )
+			{
+				NumElements = FMath::Max( NumElements, ElementID.GetValue() + 1 );
+			}
+		}
+		else
+		{
+			// Otherwise perform the remap, and get the number of elements from the resulting attribute indices array.
+			for( auto& AttributeNameAndIndicesArray : Map )
+			{
+				AttributeNameAndIndicesArray.Value.Remap( IndexRemap );
+				NumElements = AttributeNameAndIndicesArray.Value.GetNumElements();
+			}
+		}
+	}
+
+private:
+	friend FArchive& operator<<( FArchive& Ar, TAttributesMap& AttributesMap )
+	{
+		struct FRegisteredAttribute
+		{
+			FName Name;
+			int32 NumIndices;
+			AttributeType DefaultValue;
+			EMeshAttributeFlags Flags;
+		};
+
+		// If we are loading, get a list of currently registered attributes and their metadata.
+		// This will be used to update the loaded attribute map after serialization.
+		TArray<FRegisteredAttribute> RegisteredAttributes;
+		if( Ar.IsLoading() )
+		{
+			for( const auto& AttributeNameAndIndicesArray : AttributesMap.Map )
+			{
+				RegisteredAttributes.Emplace();
+				FRegisteredAttribute& RegisteredAttribute = RegisteredAttributes.Last();
+				RegisteredAttribute.Name = AttributeNameAndIndicesArray.Key;
+				RegisteredAttribute.NumIndices = AttributeNameAndIndicesArray.Value.GetNumIndices();
+				RegisteredAttribute.DefaultValue = AttributeNameAndIndicesArray.Value.GetDefaultValue();
+				RegisteredAttribute.Flags = AttributeNameAndIndicesArray.Value.GetFlags();
+			}
+		}
+
+		// First serialize the number of elements which each attribute should contain
+		Ar << AttributesMap.NumElements;
+
+		// Now serialize the attributes of this type.
+		Ar << AttributesMap.Map;
+
+		// If we are loading, the attributes just serialized may be out of sync with the currently registered attributes, so we need to fix this up now.
+		if( Ar.IsLoading() )
+		{
+			// Go through the list of registered attributes we just built
+			for( const FRegisteredAttribute& RegisteredAttribute : RegisteredAttributes )
+			{
+				AttributeIndicesArrayType* AttributeIndicesArray = AttributesMap.Map.Find( RegisteredAttribute.Name );
+				if( !AttributeIndicesArray )
+				{
+					// Re-register any attributes that aren't there
+					AttributeIndicesArrayType NewAttributeIndicesArray( RegisteredAttribute.NumIndices, RegisteredAttribute.DefaultValue, RegisteredAttribute.Flags, AttributesMap.NumElements );
+
+					UE_LOG( LogMeshDescription, Warning, TEXT( "Didn't find attribute '%s' - adding" ), *RegisteredAttribute.Name.ToString() );
+					AttributesMap.Map.Add( RegisteredAttribute.Name, MoveTemp( NewAttributeIndicesArray ) );
+				}
+				else if( AttributeIndicesArray->GetNumIndices() != RegisteredAttribute.NumIndices )
+				{
+					// Amend the number of attribute indices for any which are not correct
+					UE_LOG( LogMeshDescription, Warning, TEXT( "Found attribute '%s' with the wrong number of indices - amending" ), *RegisteredAttribute.Name.ToString() );
+					AttributeIndicesArray->SetNumIndices( RegisteredAttribute.NumIndices );
+				}
+				else
+				{
+					UE_LOG( LogMeshDescription, Log, TEXT( "Found attribute '%s' OK" ), *RegisteredAttribute.Name.ToString() );
+					check( AttributeIndicesArray->GetArrayForIndex( 0 ).Num() == AttributesMap.NumElements );
+				}
+			}
+
+			// Finally, go through all the attributes in the container, and remove any entries which are not in the registered list
+			for( auto AttributesIt = AttributesMap.Map.CreateIterator(); AttributesIt; ++AttributesIt )
+			{
+				const FName& AttributeName = AttributesIt.Key();
+				if( !RegisteredAttributes.ContainsByPredicate(
+					[ &AttributeName ]( const FRegisteredAttribute& RegisteredAttribute ) { return RegisteredAttribute.Name == AttributeName; }
+				) )
+				{
+					UE_LOG( LogMeshDescription, Warning, TEXT( "Found unregistered attribute '%s' - removing" ), *AttributeName.ToString() );
+					AttributesIt.RemoveCurrent();
+				}
+			}
+		}
+
+		return Ar;
+	}
+
+	/** Number of elements for each attribute index */
+	int32 NumElements;
+
+	/** The actual container */
+	MapType Map;
+};
 
 
 /**
@@ -325,6 +612,18 @@ struct TTupleIndex<Type, TTuple<Head, Tail...>>
 };
 
 
+/** Helper template which splits a tuple into head and tail */
+template <typename Tuple>
+struct TSplitTuple;
+
+template <typename TupleHead, typename... TupleTail>
+struct TSplitTuple<TTuple<TupleHead, TupleTail...>>
+{
+	using Head = TupleHead;
+	using Tail = TTuple<TupleTail...>;
+};
+
+
 /**
  * This is the container for all attributes of a particular mesh element.
  * It contains a TTuple of TAttributesMap, one per attribute type (FVector, float, bool, etc)
@@ -333,7 +632,6 @@ template <typename ElementIDType>
 class TAttributesSet
 {
 public:
-
 	/**
 	 * Register a new attribute name with the given type (must be a member of the AttributeTypes tuple).
 	 *
@@ -347,35 +645,21 @@ public:
 	template <typename AttributeType>
 	void RegisterAttribute( const FName AttributeName, const int32 NumberOfIndices = 1, const AttributeType& Default = AttributeType(), const EMeshAttributeFlags Flags = EMeshAttributeFlags::None )
 	{
-		// See if there is an existing attribute in this set, and get the number of elements it has.
-		// When a new attribute is registered in a non-empty attribute set, it needs to be default-initialized to contain the same
-		// number of elements as its contemporaries.
-		TAttributesMap<AttributeType, ElementIDType>& AttributesMap = Container.Get<TTupleIndex<AttributeType, AttributeTypes>::Value>();
-		typename TAttributesMap<AttributeType, ElementIDType>::TConstIterator It( AttributesMap.CreateConstIterator() );
-		const int32 NumElements = ( !It ) ? 0 : It.Value().GetArrayForIndex( 0 ).Num();
-
-		if( !AttributesMap.Contains( AttributeName ) )
-		{
-			TAttributeIndicesArray<AttributeType, ElementIDType>& AttributeIndicesArray = AttributesMap.Emplace( AttributeName, TAttributeIndicesArray<AttributeType, ElementIDType>( NumberOfIndices, Default, Flags ) );
-			if( NumElements > 0 )
-			{
-				AttributeIndicesArray.Insert( ElementIDType( NumElements - 1 ) );
-			}
-		}
+		Container.Get<TTupleIndex<AttributeType, AttributeTypes>::Value>().RegisterAttribute( AttributeName, NumberOfIndices, Default, Flags );
 	}
 
 	/** Unregister an attribute name with the given type */
 	template <typename AttributeType>
 	void UnregisterAttribute( const FName AttributeName )
 	{
-		Container.Get<TTupleIndex<AttributeType, AttributeTypes>::Value>().Remove( AttributeName );
+		Container.Get<TTupleIndex<AttributeType, AttributeTypes>::Value>().UnregisterAttribute( AttributeName );
 	}
 
 	/** Determines whether an attribute of the given type exists with the given name */
 	template <typename AttributeType>
 	bool HasAttribute( const FName AttributeName ) const
 	{
-		return Container.Get<TTupleIndex<AttributeType, AttributeTypes>::Value>().Contains( AttributeName );
+		return Container.Get<TTupleIndex<AttributeType, AttributeTypes>::Value>().HasAttribute( AttributeName );
 	}
 
 	/**
@@ -394,14 +678,14 @@ public:
 	TMeshAttributeArray<AttributeType, ElementIDType>& GetAttributes( const FName AttributeName, const int32 AttributeIndex = 0 )
 	{
 		// @todo mesh description: should this handle non-existent attribute names and indices gracefully?
-		return Container.Get<TTupleIndex<AttributeType, AttributeTypes>::Value>().FindChecked( AttributeName ).GetArrayForIndex( AttributeIndex );
+		return Container.Get<TTupleIndex<AttributeType, AttributeTypes>::Value>().GetAttributes( AttributeName, AttributeIndex );
 	}
 
 	template <typename AttributeType>
 	const TMeshAttributeArray<AttributeType, ElementIDType>& GetAttributes( const FName AttributeName, const int32 AttributeIndex = 0 ) const
 	{
 		// @todo mesh description: should this handle non-existent attribute names and indices gracefully?
-		return Container.Get<TTupleIndex<AttributeType, AttributeTypes>::Value>().FindChecked( AttributeName ).GetArrayForIndex( AttributeIndex );
+		return Container.Get<TTupleIndex<AttributeType, AttributeTypes>::Value>().GetAttributes( AttributeName, AttributeIndex );
 	}
 
 	/**
@@ -421,14 +705,14 @@ public:
 	TAttributeIndicesArray<AttributeType, ElementIDType>& GetAttributesSet( const FName AttributeName )
 	{
 		// @todo mesh description: should this handle non-existent attribute names gracefully?
-		return Container.Get<TTupleIndex<AttributeType, AttributeTypes>::Value>().FindChecked( AttributeName );
+		return Container.Get<TTupleIndex<AttributeType, AttributeTypes>::Value>().GetAttributesSet( AttributeName );
 	}
 
 	template <typename AttributeType>
 	const TAttributeIndicesArray<AttributeType, ElementIDType>& GetAttributesSet( const FName AttributeName ) const
 	{
 		// @todo mesh description: should this handle non-existent attribute names gracefully?
-		return Container.Get<TTupleIndex<AttributeType, AttributeTypes>::Value>().FindChecked( AttributeName );
+		return Container.Get<TTupleIndex<AttributeType, AttributeTypes>::Value>().GetAttributesSet( AttributeName );
 	}
 
 	/** Returns the number of indices for the attribute with the given name */
@@ -436,61 +720,51 @@ public:
 	int32 GetAttributeIndexCount( const FName AttributeName ) const
 	{
 		// @todo mesh description: should this handle non-existent attribute names and indices gracefully?
-		return Container.Get<TTupleIndex<AttributeType, AttributeTypes>::Value>().FindChecked( AttributeName ).GetNumIndices();
+		return Container.Get<TTupleIndex<AttributeType, AttributeTypes>::Value>().GetAttributeIndexCount( AttributeName );
 	}
 
 	/** Sets the number of indices for the attribute with the given name */
 	template <typename AttributeType>
 	void SetAttributeIndexCount( const FName AttributeName, const int32 NumIndices )
 	{
-		Container.Get<TTupleIndex<AttributeType, AttributeTypes>::Value>().FindChecked( AttributeName ).SetNumIndices( NumIndices );
+		Container.Get<TTupleIndex<AttributeType, AttributeTypes>::Value>().SetAttributeIndexCount( AttributeName, NumIndices );
 	}
 
 	/** Returns an array of all the attribute names registered for this attribute type */
 	template <typename AttributeType, typename Allocator>
 	void GetAttributeNames( TArray<FName, Allocator>& OutAttributeNames ) const
 	{
-		Container.Get<TTupleIndex<AttributeType, AttributeTypes>::Value>().GetKeys( OutAttributeNames );
+		Container.Get<TTupleIndex<AttributeType, AttributeTypes>::Value>().GetAttributeNames( OutAttributeNames );
 	}
 
 	template <typename AttributeType>
 	AttributeType GetAttribute( const ElementIDType ElementID, const FName AttributeName, const int32 AttributeIndex = 0 ) const
 	{
-		return Container.Get<TTupleIndex<AttributeType, AttributeTypes>::Value>().FindChecked( AttributeName ).GetArrayForIndex( AttributeIndex )[ ElementID ];
+		return Container.Get<TTupleIndex<AttributeType, AttributeTypes>::Value>().GetAttribute( ElementID, AttributeName, AttributeIndex );
 	}
 
 	template <typename AttributeType>
 	void SetAttribute( const ElementIDType ElementID, const FName AttributeName, const int32 AttributeIndex, const AttributeType& AttributeValue )
 	{
-		Container.Get<TTupleIndex<AttributeType, AttributeTypes>::Value>().FindChecked( AttributeName ).GetArrayForIndex( AttributeIndex )[ ElementID ] = AttributeValue;
+		Container.Get<TTupleIndex<AttributeType, AttributeTypes>::Value>().SetAttribute( ElementID, AttributeName, AttributeIndex, AttributeValue );
 	}
 
 	/** Inserts a default-initialized value for all attributes of the given ID */
 	void Insert( const ElementIDType ElementID )
 	{
-		auto AddElements = [ ElementID ]( auto& AttributesMap )
-		{
-			for( auto& AttributeNameAndIndicesArray : AttributesMap )
-			{
-				AttributeNameAndIndicesArray.Value.Insert( ElementID );
-			}
-		};
-
-		VisitTupleElements( Container, AddElements );
+		VisitTupleElements( Container, [ ElementID ]( auto& AttributesMap ) { AttributesMap.Insert( ElementID ); } );
 	}
 
 	/** Removes all attributes with the given ID */
 	void Remove( const ElementIDType ElementID )
 	{
-		auto RemoveElements = [ ElementID ]( auto& AttributesMap )
-		{
-			for( auto& AttributeNameAndIndicesArray : AttributesMap )
-			{
-				AttributeNameAndIndicesArray.Value.Remove( ElementID );
-			}
-		};
+		VisitTupleElements( Container, [ ElementID ]( auto& AttributesMap ) { AttributesMap.Remove( ElementID ); } );
+	}
 
-		VisitTupleElements( Container, RemoveElements );
+	/** Initializes the attribute set with the given number of elements, all at the default value */
+	void Initialize( const int32 NumElements )
+	{
+		VisitTupleElements( Container, [ NumElements ]( auto& AttributesMap ) { AttributesMap.Initialize( NumElements ); } );
 	}
 
 	/**
@@ -500,147 +774,74 @@ public:
 	template <typename FuncType>
 	void ForEachAttributeIndicesArray( const FuncType& Func )
 	{
-		auto ForEach = [ &Func ]( auto& AttributesMap )
-		{
-			for( auto& AttributeNameAndIndicesArray : AttributesMap )
-			{
-				Func( AttributeNameAndIndicesArray.Key, AttributeNameAndIndicesArray.Value );
-			}
-		};
-
-		VisitTupleElements( Container, ForEach );
+		VisitTupleElements( Container, [ &Func ]( auto& AttributesMap ) { AttributesMap.ForEachAttributeIndicesArray( Func ); } );
 	}
 
 	template <typename FuncType>
 	void ForEachAttributeIndicesArray( const FuncType& Func ) const
 	{
-		auto ForEach = [ &Func ]( const auto& AttributesMap )
-		{
-			for( const auto& AttributeNameAndIndicesArray : AttributesMap )
-			{
-				Func( AttributeNameAndIndicesArray.Key, AttributeNameAndIndicesArray.Value );
-			}
-		};
-
-		VisitTupleElements( Container, ForEach );
+		VisitTupleElements( Container, [ &Func ]( const auto& AttributesMap ) { AttributesMap.ForEachAttributeIndicesArray( Func ); } );
 	}
 
 	/** Applies the given remapping to the attributes set */
 	void Remap( const TSparseArray<ElementIDType>& IndexRemap )
 	{
-		auto RemapElements = [ &IndexRemap ]( auto& AttributesMap )
-		{
-			for( auto& AttributeNameAndIndicesArray : AttributesMap )
-			{
-				AttributeNameAndIndicesArray.Value.Remap( IndexRemap );
-			}
-		};
-
-		VisitTupleElements( Container, RemapElements );
+		VisitTupleElements( Container, [ &IndexRemap ]( auto& AttributesMap ) { AttributesMap.Remap( IndexRemap ); } );
 	}
 
 	/** Serializer */
 	friend FArchive& operator<<( FArchive& Ar, TAttributesSet& AttributesSet )
 	{
-		auto SerializeElements = [ &Ar ]( auto& AttributesMap )
+		// Serialize the number of attribute types in the container tuple.
+		// If loading, this may be different to the current number defined.
+		const int32 NumAttributeTypes = TTupleArity<AttributeTypes>::Value;
+		int32 SerializedAttributeTypes = NumAttributeTypes;
+		Ar << SerializedAttributeTypes;
+		// Cannot deserialize more attribute types than there are
+		check( NumAttributeTypes >= SerializedAttributeTypes );
+
+		// Serialize the tuple of attribute maps by hand, so we can deserialize correctly when the archive contains fewer tuple elements than the current code
+		// NOTE: This relies on the assumption that VisitTupleElements will always visit elements in ascending order.
+		VisitTupleElements( AttributesSet.Container, [ &Ar, SerializedAttributeTypes, TypeIndex = 0, NumElements = 0 ]( auto& AttributesMap ) mutable
 		{
-			using AttributeIndicesArrayType = decltype( DeclVal<typename TDecay<decltype( AttributesMap )>::Type::ElementType>().Value );
-			using AttributeType = typename AttributeIndicesArrayType::AttributeType;
-
-			struct FRegisteredAttribute
-			{
-				FName Name;
-				int32 NumIndices;
-				AttributeType DefaultValue;
-				EMeshAttributeFlags Flags;
-			};
-
-			// If we are loading, get a list of currently registered attributes and their metadata.
-			// This will be used to update the loaded attribute map after serialization.
-			TArray<FRegisteredAttribute> RegisteredAttributes;
-			if( Ar.IsLoading() )
+			if( TypeIndex < SerializedAttributeTypes )
 			{
-				for( const auto& AttributeNameAndIndicesArray : AttributesMap )
-				{
-					RegisteredAttributes.Emplace();
-					FRegisteredAttribute& RegisteredAttribute = RegisteredAttributes.Last();
-					RegisteredAttribute.Name = AttributeNameAndIndicesArray.Key;
-					RegisteredAttribute.NumIndices = AttributeNameAndIndicesArray.Value.GetNumIndices();
-					RegisteredAttribute.DefaultValue = AttributeNameAndIndicesArray.Value.GetDefaultValue();
-					RegisteredAttribute.Flags = AttributeNameAndIndicesArray.Value.GetFlags();
-				}
+				// Serialize attributes map, and keep note of the number of elements present.
+				// This should be the same every iteration.
+				Ar << AttributesMap;
+				check( TypeIndex == 0 || NumElements == AttributesMap.GetNumElements() );
+				NumElements = AttributesMap.GetNumElements();
 			}
-
-			// Now serialize the attributes of this type.
-			Ar << AttributesMap;
-
-			// If we are loading, the attributes just serialized may be out of sync with the currently registered attributes, so we need to fix this up now.
-			if( Ar.IsLoading() )
+			else
 			{
-				// Determine the number of elements for the attributes just serialized
-				auto It( AttributesMap.CreateConstIterator() );
-				const int32 NumElements = ( !It ) ? 0 : It.Value().GetArrayForIndex( 0 ).Num();
-
-				// Go through the list of registered attributes we just built
-				for( const FRegisteredAttribute& RegisteredAttribute : RegisteredAttributes )
-				{
-					AttributeIndicesArrayType* AttributeIndicesArray = AttributesMap.Find( RegisteredAttribute.Name );
-					if( !AttributeIndicesArray )
-					{
-						// Re-register any attributes that aren't there
-						AttributeIndicesArrayType NewAttributeIndicesArray( RegisteredAttribute.NumIndices, RegisteredAttribute.DefaultValue, RegisteredAttribute.Flags );
-						if( NumElements > 0 )
-						{
-							// and fill the new attribute with default initialized values
-							NewAttributeIndicesArray.Insert( ElementIDType( NumElements - 1 ) );
-						}
-
-						UE_LOG( LogTemp, Warning, TEXT( "Didn't find attribute '%s' - adding" ), *RegisteredAttribute.Name.ToString() );
-						AttributesMap.Add( RegisteredAttribute.Name, MoveTemp( NewAttributeIndicesArray ) );
-					}
-					else if( AttributeIndicesArray->GetNumIndices() != RegisteredAttribute.NumIndices )
-					{
-						// Amend the number of attribute indices for any which are not correct
-						UE_LOG( LogTemp, Warning, TEXT( "Found attribute '%s' with the wrong number of indices - amending" ), *RegisteredAttribute.Name.ToString() );
-						AttributeIndicesArray->SetNumIndices( RegisteredAttribute.NumIndices );
-					}
-				}
-
-				// Finally, go through all the attributes in the container, and remove any entries which are not in the registered list
-				for( auto AttributesIt = AttributesMap.CreateIterator(); AttributesIt; ++AttributesIt )
-				{
-					const FName& AttributeName = AttributesIt.Key();
-					if( !RegisteredAttributes.ContainsByPredicate(
-						[ &AttributeName ]( const FRegisteredAttribute& RegisteredAttribute ) { return RegisteredAttribute.Name == AttributeName; }
-					  ) )
-					{
-						UE_LOG( LogTemp, Warning, TEXT( "Found unregistered attribute '%s' - removing" ), *AttributeName.ToString() );
-						AttributesIt.RemoveCurrent();
-					}
-				}
+				// If we have run out of data to deserialize, initialize this attributes map so that it has the same number of elements
+				// as the other maps.
+				check( Ar.IsLoading() );
+				AttributesMap.Initialize( NumElements );
 			}
-		};
-
-		VisitTupleElements( AttributesSet.Container, SerializeElements );
+			TypeIndex++;
+		}
+		);
 
 		return Ar;
 	}
 
 private:
 	/**
-	  * Define type for the entire attribute container.
-	  * We can have attributes of multiple types, each with a name and an arbitrary number of indices,
-	  * whose elements are indexed by an ElementIDType.
-	  *
-	  * This implies the below data structure:
-	  * A TTuple (one per attribute type) of
-	  * TMap keyed on the attribute name,
-	  * yielding a TArray indexed by attribute index,
-	  * yielding a TMeshAttributeArray indexed by an Element ID,
-	  * yielding an item of type AttributeType.
-	  *
-	  * This looks complicated, but actually makes attribute lookup easy when we are interested in a particular attribute for many element IDs.
-	  * By caching the TMeshAttributeArray arrived at by the attribute name and index, we have O(1) access to that attribute for all elements.
-	  */
-	typename TMakeAttributesSet<ElementIDType, AttributeTypes>::Type Container;
+	 * Define type for the entire attribute container.
+	 * We can have attributes of multiple types, each with a name and an arbitrary number of indices,
+	 * whose elements are indexed by an ElementIDType.
+	 *
+	 * This implies the below data structure:
+	 * A TTuple (one per attribute type) of
+	 * TMap keyed on the attribute name,
+	 * yielding a TArray indexed by attribute index,
+	 * yielding a TMeshAttributeArray indexed by an Element ID,
+	 * yielding an item of type AttributeType.
+	 *
+	 * This looks complicated, but actually makes attribute lookup easy when we are interested in a particular attribute for many element IDs.
+	 * By caching the TMeshAttributeArray arrived at by the attribute name and index, we have O(1) access to that attribute for all elements.
+	 */
+	using ContainerType = typename TMakeAttributesSet<ElementIDType, AttributeTypes>::Type;
+	ContainerType Container;
 };
diff --git a/Engine/Source/Runtime/MeshDescription/Public/MeshDescription.h b/Engine/Source/Runtime/MeshDescription/Public/MeshDescription.h
--- a/Engine/Source/Runtime/MeshDescription/Public/MeshDescription.h
+++ b/Engine/Source/Runtime/MeshDescription/Public/MeshDescription.h
@@ -11,6 +11,9 @@
 #include "MeshDescription.generated.h"
 
 
+MESHDESCRIPTION_API DECLARE_LOG_CATEGORY_EXTERN( LogMeshDescription, Log, All );
+
+
 USTRUCT()
 struct FMeshVertex
 {
