Added TAttributesSet<>::RegisterAttributesFromAttributesSet().
This copies over all attributes from an existing attributes set to this one.

diff --git a/Engine/Source/Runtime/MeshDescription/Public/MeshAttributeArray.h b/Engine/Source/Runtime/MeshDescription/Public/MeshAttributeArray.h
--- a/Engine/Source/Runtime/MeshDescription/Public/MeshAttributeArray.h
+++ b/Engine/Source/Runtime/MeshDescription/Public/MeshAttributeArray.h
@@ -461,6 +461,18 @@ public:
 		}
 	}
 
+	/** Registers attributes copied from the specified attributes map */
+	void RegisterAttributesFromAttributesSet( const TAttributesMap& SrcAttributesMap )
+	{
+		for( const auto& AttributeNameAndIndicesArray : SrcAttributesMap.Map )
+		{
+			const FName& AttributeName = AttributeNameAndIndicesArray.Key;
+			const AttributeIndicesArrayType& AttributeIndicesArray = AttributeNameAndIndicesArray.Value;
+
+			RegisterAttribute( AttributeName, AttributeIndicesArray.GetNumIndices(), AttributeIndicesArray.GetDefaultValue(), AttributeIndicesArray.GetFlags() );
+		}
+	}
+
 	/** Applies the given remapping to the attributes set */
 	void Remap( const TSparseArray<ElementIDType>& IndexRemap )
 	{
@@ -789,6 +801,13 @@ public:
 		VisitTupleElements( [ &IndexRemap ]( auto& AttributesMap ) { AttributesMap.Remap( IndexRemap ); }, Container );
 	}
 
+	/** Copies registered attributes from another TAttributesSet */
+	void RegisterAttributesFromAttributesSet( const TAttributesSet& Other )
+	{
+		VisitTupleElements( [ this ]( auto& DestAttributesMap, const auto& SrcAttributesMap ) { DestAttributesMap.RegisterAttributesFromAttributesSet( SrcAttributesMap ); },
+			Container, Other.Container );
+	}
+
 	/** Serializer */
 	friend FArchive& operator<<( FArchive& Ar, TAttributesSet& AttributesSet )
 	{
@@ -803,26 +822,25 @@ public:
 		// Serialize the tuple of attribute maps by hand, so we can deserialize correctly when the archive contains fewer tuple elements than the current code
 		// NOTE: This relies on the assumption that VisitTupleElements will always visit elements in ascending order.
 		VisitTupleElements( [ &Ar, SerializedAttributeTypes, TypeIndex = 0, NumElements = 0 ]( auto& AttributesMap ) mutable
-		{
-			if( TypeIndex < SerializedAttributeTypes )
-			{
-				// Serialize attributes map, and keep note of the number of elements present.
-				// This should be the same every iteration.
-				Ar << AttributesMap;
-				check( TypeIndex == 0 || NumElements == AttributesMap.GetNumElements() );
-				NumElements = AttributesMap.GetNumElements();
-			}
-			else
 			{
-				// If we have run out of data to deserialize, initialize this attributes map so that it has the same number of elements
-				// as the other maps.
-				check( Ar.IsLoading() );
-				AttributesMap.Initialize( NumElements );
-			}
-			TypeIndex++;
-		},
-		AttributesSet.Container
-		);
+				if( TypeIndex < SerializedAttributeTypes )
+				{
+					// Serialize attributes map, and keep note of the number of elements present.
+					// This should be the same every iteration.
+					Ar << AttributesMap;
+					check( TypeIndex == 0 || NumElements == AttributesMap.GetNumElements() );
+					NumElements = AttributesMap.GetNumElements();
+				}
+				else
+				{
+					// If we have run out of data to deserialize, initialize this attributes map so that it has the same number of elements
+					// as the other maps.
+					check( Ar.IsLoading() );
+					AttributesMap.Initialize( NumElements );
+				}
+				TypeIndex++;
+			},
+			AttributesSet.Container );
 
 		return Ar;
 	}
