Initial attempt to fix Linux build errors.

diff --git a/Engine/Source/Runtime/MeshDescription/Private/MeshDescription.cpp b/Engine/Source/Runtime/MeshDescription/Private/MeshDescription.cpp
--- a/Engine/Source/Runtime/MeshDescription/Private/MeshDescription.cpp
+++ b/Engine/Source/Runtime/MeshDescription/Private/MeshDescription.cpp
@@ -7,9 +7,6 @@
 #include "Serialization/ObjectWriter.h"
 
 
-DEFINE_LOG_CATEGORY( LogMeshDescription );
-
-
 UMeshDescription::UMeshDescription()
 {
 	RegisterBaseAttributes();
diff --git a/Engine/Source/Runtime/MeshDescription/Private/MeshTypes.cpp b/Engine/Source/Runtime/MeshDescription/Private/MeshTypes.cpp
--- a/Engine/Source/Runtime/MeshDescription/Private/MeshTypes.cpp
+++ b/Engine/Source/Runtime/MeshDescription/Private/MeshTypes.cpp
@@ -2,6 +2,9 @@
 
 #include "MeshTypes.h"
 
+DEFINE_LOG_CATEGORY( LogMeshDescription );
+
+
 const FElementID FElementID::Invalid( TNumericLimits<uint32>::Max() );
 const FVertexID FVertexID::Invalid( TNumericLimits<uint32>::Max() );
 const FVertexInstanceID FVertexInstanceID::Invalid( TNumericLimits<uint32>::Max() );
diff --git a/Engine/Source/Runtime/MeshDescription/Public/MeshAttributeArray.h b/Engine/Source/Runtime/MeshDescription/Public/MeshAttributeArray.h
--- a/Engine/Source/Runtime/MeshDescription/Public/MeshAttributeArray.h
+++ b/Engine/Source/Runtime/MeshDescription/Public/MeshAttributeArray.h
@@ -144,7 +144,7 @@ void TMeshAttributeArray<ElementType, ElementIDType>::Remap( const TSparseArray<
 {
 	TMeshAttributeArray NewAttributeArray;
 
-	for( TSparseArray<ElementIDType>::TConstIterator It( IndexRemap ); It; ++It )
+	for( typename TSparseArray<ElementIDType>::TConstIterator It( IndexRemap ); It; ++It )
 	{
 		const int32 OldElementIndex = It.GetIndex();
 		const ElementIDType NewElementIndex = IndexRemap[ OldElementIndex ];
@@ -311,12 +311,12 @@ template <typename AttributeType> using TPolygonGroupAttributeIndicesArray = TAt
 /**
  * This maps an attribute name to a TAttributeIndicesArray, i.e. an array of MeshAttributeArrays, one per attribute index.
  */
-template <typename Attribute, typename ElementID>
+template <typename T, typename U>
 class TAttributesMap
 {
 public:
-	using AttributeType = Attribute;
-	using ElementIDType = ElementID;
+	using AttributeType = T;
+	using ElementIDType = U;
 	using AttributeIndicesArrayType = TAttributeIndicesArray<AttributeType, ElementIDType>;
 	using MapType = TMap<FName, AttributeIndicesArrayType>;
 
diff --git a/Engine/Source/Runtime/MeshDescription/Public/MeshDescription.h b/Engine/Source/Runtime/MeshDescription/Public/MeshDescription.h
--- a/Engine/Source/Runtime/MeshDescription/Public/MeshDescription.h
+++ b/Engine/Source/Runtime/MeshDescription/Public/MeshDescription.h
@@ -11,9 +11,6 @@
 #include "MeshDescription.generated.h"
 
 
-MESHDESCRIPTION_API DECLARE_LOG_CATEGORY_EXTERN( LogMeshDescription, Log, All );
-
-
 USTRUCT()
 struct FMeshVertex
 {
diff --git a/Engine/Source/Runtime/MeshDescription/Public/MeshTypes.h b/Engine/Source/Runtime/MeshDescription/Public/MeshTypes.h
--- a/Engine/Source/Runtime/MeshDescription/Public/MeshTypes.h
+++ b/Engine/Source/Runtime/MeshDescription/Public/MeshTypes.h
@@ -5,6 +5,8 @@
 #include "CoreMinimal.h"
 #include "MeshTypes.generated.h"
 
+MESHDESCRIPTION_API DECLARE_LOG_CATEGORY_EXTERN( LogMeshDescription, Log, All );
+
 
 // @todo mesheditor: Need comments
 
