Further Linux build fixes.

diff --git a/Engine/Source/Runtime/MeshDescription/Public/MeshElementArray.h b/Engine/Source/Runtime/MeshDescription/Public/MeshElementArray.h
--- a/Engine/Source/Runtime/MeshDescription/Public/MeshElementArray.h
+++ b/Engine/Source/Runtime/MeshDescription/Public/MeshElementArray.h
@@ -264,7 +264,7 @@ void TMeshElementArray<ElementType, ElementIDType>::Compact( TSparseArray<Elemen
 	OutIndexRemap.Empty( Container.GetMaxIndex() );
 
 	// Add valid elements into a new contiguous sparse array.  Note non-const iterator so we can move elements.
-	for( TSparseArray<ElementType>::TIterator It( Container ); It; ++It )
+	for( typename TSparseArray<ElementType>::TIterator It( Container ); It; ++It )
 	{
 		const int32 OldElementIndex = It.GetIndex();
 
@@ -287,7 +287,7 @@ void TMeshElementArray<ElementType, ElementIDType>::Remap( const TSparseArray<El
 	NewContainer.Empty( IndexRemap.GetMaxIndex() );
 
 	// Add valid elements into a new contiguous sparse array.  Note non-const iterator so we can move elements.
-	for( TSparseArray<ElementType>::TIterator It( Container ); It; ++It )
+	for( typename TSparseArray<ElementType>::TIterator It( Container ); It; ++It )
 	{
 		const int32 OldElementIndex = It.GetIndex();
 
