Attempt to get more information out the Linux build as to what's failing.

diff --git a/Engine/Source/Runtime/MeshDescription/Public/MeshAttributeArray.h b/Engine/Source/Runtime/MeshDescription/Public/MeshAttributeArray.h
--- a/Engine/Source/Runtime/MeshDescription/Public/MeshAttributeArray.h
+++ b/Engine/Source/Runtime/MeshDescription/Public/MeshAttributeArray.h
@@ -41,10 +41,10 @@ class TMeshAttributeArrayBase
 public:
 
 	/** Disallow use of the copy constructor to prevent arrays being mistakenly accessed in the UMeshDescription by value */
-	TMeshAttributeArrayBase( const TMeshAttributeArrayBase& )
-	{
-		static_assert( false, "TMeshAttributeArray must be accessed by reference");
-	}
+	TMeshAttributeArrayBase( const TMeshAttributeArrayBase& ) = delete;
+//	{
+//		static_assert( false, "TMeshAttributeArray must be accessed by reference");
+//	}
 
 	/**
 	 * Custom serialization for TMeshAttributeArrayBase.
