Fix some unix build error

diff --git a/Engine/Source/Runtime/Engine/Private/StaticMesh.cpp b/Engine/Source/Runtime/Engine/Private/StaticMesh.cpp
--- a/Engine/Source/Runtime/Engine/Private/StaticMesh.cpp
+++ b/Engine/Source/Runtime/Engine/Private/StaticMesh.cpp
@@ -2167,7 +2167,7 @@ FStaticMeshSourceModel::~FStaticMeshSourceModel()
 #if WITH_EDITOR
 bool FStaticMeshSourceModel::IsRawMeshEmpty() const
 {
-	return (RawMeshBulkData == nullptr || RawMeshBulkData->IsEmpty() && OriginalMeshDescription == nullptr);
+	return (RawMeshBulkData == nullptr || (RawMeshBulkData->IsEmpty() && OriginalMeshDescription == nullptr));
 }
 
 void FStaticMeshSourceModel::LoadRawMesh(FRawMesh& OutRawMesh) const
diff --git a/Engine/Source/Runtime/MeshDescription/Public/MeshAttributeArray.h b/Engine/Source/Runtime/MeshDescription/Public/MeshAttributeArray.h
--- a/Engine/Source/Runtime/MeshDescription/Public/MeshAttributeArray.h
+++ b/Engine/Source/Runtime/MeshDescription/Public/MeshAttributeArray.h
@@ -657,21 +657,21 @@ public:
 	template <typename AttributeType>
 	void RegisterAttribute( const FName AttributeName, const int32 NumberOfIndices = 1, const AttributeType& Default = AttributeType(), const EMeshAttributeFlags Flags = EMeshAttributeFlags::None )
 	{
-		Container.Get<TTupleIndex<AttributeType, AttributeTypes>::Value>().RegisterAttribute( AttributeName, NumberOfIndices, Default, Flags );
+		Container.template Get<typename TTupleIndex<AttributeType, AttributeTypes>::Value>().RegisterAttribute( AttributeName, NumberOfIndices, Default, Flags );
 	}
 
 	/** Unregister an attribute name with the given type */
 	template <typename AttributeType>
 	void UnregisterAttribute( const FName AttributeName )
 	{
-		Container.Get<TTupleIndex<AttributeType, AttributeTypes>::Value>().UnregisterAttribute( AttributeName );
+		Container.template Get<typename TTupleIndex<AttributeType, AttributeTypes>::Value>().UnregisterAttribute( AttributeName );
 	}
 
 	/** Determines whether an attribute of the given type exists with the given name */
 	template <typename AttributeType>
 	bool HasAttribute( const FName AttributeName ) const
 	{
-		return Container.Get<TTupleIndex<AttributeType, AttributeTypes>::Value>().HasAttribute( AttributeName );
+		return Container.template Get<typename TTupleIndex<AttributeType, AttributeTypes>::Value>().HasAttribute( AttributeName );
 	}
 
 	/**
@@ -690,14 +690,14 @@ public:
 	TMeshAttributeArray<AttributeType, ElementIDType>& GetAttributes( const FName AttributeName, const int32 AttributeIndex = 0 )
 	{
 		// @todo mesh description: should this handle non-existent attribute names and indices gracefully?
-		return Container.Get<TTupleIndex<AttributeType, AttributeTypes>::Value>().GetAttributes( AttributeName, AttributeIndex );
+		return Container.template Get<typename TTupleIndex<AttributeType, AttributeTypes>::Value>().GetAttributes( AttributeName, AttributeIndex );
 	}
 
 	template <typename AttributeType>
 	const TMeshAttributeArray<AttributeType, ElementIDType>& GetAttributes( const FName AttributeName, const int32 AttributeIndex = 0 ) const
 	{
 		// @todo mesh description: should this handle non-existent attribute names and indices gracefully?
-		return Container.Get<TTupleIndex<AttributeType, AttributeTypes>::Value>().GetAttributes( AttributeName, AttributeIndex );
+		return Container.template Get<typename TTupleIndex<AttributeType, AttributeTypes>::Value>().GetAttributes( AttributeName, AttributeIndex );
 	}
 
 	/**
@@ -717,14 +717,14 @@ public:
 	TAttributeIndicesArray<AttributeType, ElementIDType>& GetAttributesSet( const FName AttributeName )
 	{
 		// @todo mesh description: should this handle non-existent attribute names gracefully?
-		return Container.Get<TTupleIndex<AttributeType, AttributeTypes>::Value>().GetAttributesSet( AttributeName );
+		return Container.template Get<typename TTupleIndex<AttributeType, AttributeTypes>::Value>().GetAttributesSet( AttributeName );
 	}
 
 	template <typename AttributeType>
 	const TAttributeIndicesArray<AttributeType, ElementIDType>& GetAttributesSet( const FName AttributeName ) const
 	{
 		// @todo mesh description: should this handle non-existent attribute names gracefully?
-		return Container.Get<TTupleIndex<AttributeType, AttributeTypes>::Value>().GetAttributesSet( AttributeName );
+		return Container.template Get<typename TTupleIndex<AttributeType, AttributeTypes>::Value>().GetAttributesSet( AttributeName );
 	}
 
 	/** Returns the number of indices for the attribute with the given name */
@@ -732,33 +732,33 @@ public:
 	int32 GetAttributeIndexCount( const FName AttributeName ) const
 	{
 		// @todo mesh description: should this handle non-existent attribute names and indices gracefully?
-		return Container.Get<TTupleIndex<AttributeType, AttributeTypes>::Value>().GetAttributeIndexCount( AttributeName );
+		return Container.template Get<typename TTupleIndex<AttributeType, AttributeTypes>::Value>().GetAttributeIndexCount( AttributeName );
 	}
 
 	/** Sets the number of indices for the attribute with the given name */
 	template <typename AttributeType>
 	void SetAttributeIndexCount( const FName AttributeName, const int32 NumIndices )
 	{
-		Container.Get<TTupleIndex<AttributeType, AttributeTypes>::Value>().SetAttributeIndexCount( AttributeName, NumIndices );
+		Container.template Get<typename TTupleIndex<AttributeType, AttributeTypes>::Value>().SetAttributeIndexCount( AttributeName, NumIndices );
 	}
 
 	/** Returns an array of all the attribute names registered for this attribute type */
 	template <typename AttributeType, typename Allocator>
 	void GetAttributeNames( TArray<FName, Allocator>& OutAttributeNames ) const
 	{
-		Container.Get<TTupleIndex<AttributeType, AttributeTypes>::Value>().GetAttributeNames( OutAttributeNames );
+		Container.template Get<typename TTupleIndex<AttributeType, AttributeTypes>::Value>().GetAttributeNames( OutAttributeNames );
 	}
 
 	template <typename AttributeType>
 	AttributeType GetAttribute( const ElementIDType ElementID, const FName AttributeName, const int32 AttributeIndex = 0 ) const
 	{
-		return Container.Get<TTupleIndex<AttributeType, AttributeTypes>::Value>().GetAttribute( ElementID, AttributeName, AttributeIndex );
+		return Container.template Get<typename TTupleIndex<AttributeType, AttributeTypes>::Value>().GetAttribute( ElementID, AttributeName, AttributeIndex );
 	}
 
 	template <typename AttributeType>
 	void SetAttribute( const ElementIDType ElementID, const FName AttributeName, const int32 AttributeIndex, const AttributeType& AttributeValue )
 	{
-		Container.Get<TTupleIndex<AttributeType, AttributeTypes>::Value>().SetAttribute( ElementID, AttributeName, AttributeIndex, AttributeValue );
+		Container.template Get<typename TTupleIndex<AttributeType, AttributeTypes>::Value>().SetAttribute( ElementID, AttributeName, AttributeIndex, AttributeValue );
 	}
 
 	/** Inserts a default-initialized value for all attributes of the given ID */
