Back out revision 12 from //UE4/Dev-Geometry/Engine/Source/Runtime/MeshDescription/Public/MeshAttributeArray.h to fix Windows build.
Linux build fixes still pending.

diff --git a/Engine/Source/Runtime/MeshDescription/Public/MeshAttributeArray.h b/Engine/Source/Runtime/MeshDescription/Public/MeshAttributeArray.h
--- a/Engine/Source/Runtime/MeshDescription/Public/MeshAttributeArray.h
+++ b/Engine/Source/Runtime/MeshDescription/Public/MeshAttributeArray.h
@@ -657,21 +657,21 @@ public:
 	template <typename AttributeType>
 	void RegisterAttribute( const FName AttributeName, const int32 NumberOfIndices = 1, const AttributeType& Default = AttributeType(), const EMeshAttributeFlags Flags = EMeshAttributeFlags::None )
 	{
-		Container.template Get<typename TTupleIndex<AttributeType, AttributeTypes>::Value>().RegisterAttribute( AttributeName, NumberOfIndices, Default, Flags );
+		Container.Get<TTupleIndex<AttributeType, AttributeTypes>::Value>().RegisterAttribute( AttributeName, NumberOfIndices, Default, Flags );
 	}
 
 	/** Unregister an attribute name with the given type */
 	template <typename AttributeType>
 	void UnregisterAttribute( const FName AttributeName )
 	{
-		Container.template Get<typename TTupleIndex<AttributeType, AttributeTypes>::Value>().UnregisterAttribute( AttributeName );
+		Container.Get<TTupleIndex<AttributeType, AttributeTypes>::Value>().UnregisterAttribute( AttributeName );
 	}
 
 	/** Determines whether an attribute of the given type exists with the given name */
 	template <typename AttributeType>
 	bool HasAttribute( const FName AttributeName ) const
 	{
-		return Container.template Get<typename TTupleIndex<AttributeType, AttributeTypes>::Value>().HasAttribute( AttributeName );
+		return Container.Get<TTupleIndex<AttributeType, AttributeTypes>::Value>().HasAttribute( AttributeName );
 	}
 
 	/**
@@ -690,14 +690,14 @@ public:
 	TMeshAttributeArray<AttributeType, ElementIDType>& GetAttributes( const FName AttributeName, const int32 AttributeIndex = 0 )
 	{
 		// @todo mesh description: should this handle non-existent attribute names and indices gracefully?
-		return Container.template Get<typename TTupleIndex<AttributeType, AttributeTypes>::Value>().GetAttributes( AttributeName, AttributeIndex );
+		return Container.Get<TTupleIndex<AttributeType, AttributeTypes>::Value>().GetAttributes( AttributeName, AttributeIndex );
 	}
 
 	template <typename AttributeType>
 	const TMeshAttributeArray<AttributeType, ElementIDType>& GetAttributes( const FName AttributeName, const int32 AttributeIndex = 0 ) const
 	{
 		// @todo mesh description: should this handle non-existent attribute names and indices gracefully?
-		return Container.template Get<typename TTupleIndex<AttributeType, AttributeTypes>::Value>().GetAttributes( AttributeName, AttributeIndex );
+		return Container.Get<TTupleIndex<AttributeType, AttributeTypes>::Value>().GetAttributes( AttributeName, AttributeIndex );
 	}
 
 	/**
@@ -717,14 +717,14 @@ public:
 	TAttributeIndicesArray<AttributeType, ElementIDType>& GetAttributesSet( const FName AttributeName )
 	{
 		// @todo mesh description: should this handle non-existent attribute names gracefully?
-		return Container.template Get<typename TTupleIndex<AttributeType, AttributeTypes>::Value>().GetAttributesSet( AttributeName );
+		return Container.Get<TTupleIndex<AttributeType, AttributeTypes>::Value>().GetAttributesSet( AttributeName );
 	}
 
 	template <typename AttributeType>
 	const TAttributeIndicesArray<AttributeType, ElementIDType>& GetAttributesSet( const FName AttributeName ) const
 	{
 		// @todo mesh description: should this handle non-existent attribute names gracefully?
-		return Container.template Get<typename TTupleIndex<AttributeType, AttributeTypes>::Value>().GetAttributesSet( AttributeName );
+		return Container.Get<TTupleIndex<AttributeType, AttributeTypes>::Value>().GetAttributesSet( AttributeName );
 	}
 
 	/** Returns the number of indices for the attribute with the given name */
@@ -732,33 +732,33 @@ public:
 	int32 GetAttributeIndexCount( const FName AttributeName ) const
 	{
 		// @todo mesh description: should this handle non-existent attribute names and indices gracefully?
-		return Container.template Get<typename TTupleIndex<AttributeType, AttributeTypes>::Value>().GetAttributeIndexCount( AttributeName );
+		return Container.Get<TTupleIndex<AttributeType, AttributeTypes>::Value>().GetAttributeIndexCount( AttributeName );
 	}
 
 	/** Sets the number of indices for the attribute with the given name */
 	template <typename AttributeType>
 	void SetAttributeIndexCount( const FName AttributeName, const int32 NumIndices )
 	{
-		Container.template Get<typename TTupleIndex<AttributeType, AttributeTypes>::Value>().SetAttributeIndexCount( AttributeName, NumIndices );
+		Container.Get<TTupleIndex<AttributeType, AttributeTypes>::Value>().SetAttributeIndexCount( AttributeName, NumIndices );
 	}
 
 	/** Returns an array of all the attribute names registered for this attribute type */
 	template <typename AttributeType, typename Allocator>
 	void GetAttributeNames( TArray<FName, Allocator>& OutAttributeNames ) const
 	{
-		Container.template Get<typename TTupleIndex<AttributeType, AttributeTypes>::Value>().GetAttributeNames( OutAttributeNames );
+		Container.Get<TTupleIndex<AttributeType, AttributeTypes>::Value>().GetAttributeNames( OutAttributeNames );
 	}
 
 	template <typename AttributeType>
 	AttributeType GetAttribute( const ElementIDType ElementID, const FName AttributeName, const int32 AttributeIndex = 0 ) const
 	{
-		return Container.template Get<typename TTupleIndex<AttributeType, AttributeTypes>::Value>().GetAttribute( ElementID, AttributeName, AttributeIndex );
+		return Container.Get<TTupleIndex<AttributeType, AttributeTypes>::Value>().GetAttribute( ElementID, AttributeName, AttributeIndex );
 	}
 
 	template <typename AttributeType>
 	void SetAttribute( const ElementIDType ElementID, const FName AttributeName, const int32 AttributeIndex, const AttributeType& AttributeValue )
 	{
-		Container.template Get<typename TTupleIndex<AttributeType, AttributeTypes>::Value>().SetAttribute( ElementID, AttributeName, AttributeIndex, AttributeValue );
+		Container.Get<TTupleIndex<AttributeType, AttributeTypes>::Value>().SetAttribute( ElementID, AttributeName, AttributeIndex, AttributeValue );
 	}
 
 	/** Inserts a default-initialized value for all attributes of the given ID */
