More fixes for Linux build.

diff --git a/Engine/Plugins/Runtime/EditableMesh/Source/EditableMesh/Public/EditableMeshTypes.h b/Engine/Plugins/Runtime/EditableMesh/Source/EditableMesh/Public/EditableMeshTypes.h
--- a/Engine/Plugins/Runtime/EditableMesh/Source/EditableMesh/Public/EditableMeshTypes.h
+++ b/Engine/Plugins/Runtime/EditableMesh/Source/EditableMesh/Public/EditableMeshTypes.h
@@ -364,13 +364,6 @@ struct FMeshElementAttributeValue
 
 	/** Accessors for known types. There will be a runtime error if the attribute value is not of the specified type. */
 	template <typename T> T GetValue() const;
-	template <> FVector4 GetValue<FVector4>() const { check( Type == EMeshElementAttributeType::FVector4 ); return Value_FVector4; }
-	template <> FVector GetValue<FVector>() const { check( Type == EMeshElementAttributeType::FVector ); return Value_FVector; }
-	template <> FVector2D GetValue<FVector2D>() const { check( Type == EMeshElementAttributeType::FVector2D ); return Value_FVector2D; }
-	template <> float GetValue<float>() const { check( Type == EMeshElementAttributeType::Float ); return Value_Float; }
-	template <> int GetValue<int>() const { check( Type == EMeshElementAttributeType::Int ); return Value_Int; }
-	template <> bool GetValue<bool>() const { check( Type == EMeshElementAttributeType::Bool ); return Value_Bool; }
-	template <> FName GetValue<FName>() const { check( Type == EMeshElementAttributeType::FName ); return Value_FName; }
 
 	/**
 	 * Calls the specified polymorphic lambda, with this attribute value as its passed-in parameter.
@@ -470,6 +463,16 @@ private:
 	FName Value_FName;
 };
 
+/** Implementations of accessors for known types. These must be defined outside the class definition. */
+template <> inline FVector4 FMeshElementAttributeValue::GetValue<FVector4>() const { check( Type == EMeshElementAttributeType::FVector4 ); return Value_FVector4; }
+template <> inline FVector FMeshElementAttributeValue::GetValue<FVector>() const { check( Type == EMeshElementAttributeType::FVector ); return Value_FVector; }
+template <> inline FVector2D FMeshElementAttributeValue::GetValue<FVector2D>() const { check( Type == EMeshElementAttributeType::FVector2D ); return Value_FVector2D; }
+template <> inline float FMeshElementAttributeValue::GetValue<float>() const { check( Type == EMeshElementAttributeType::Float ); return Value_Float; }
+template <> inline int FMeshElementAttributeValue::GetValue<int>() const { check( Type == EMeshElementAttributeType::Int ); return Value_Int; }
+template <> inline bool FMeshElementAttributeValue::GetValue<bool>() const { check( Type == EMeshElementAttributeType::Bool ); return Value_Bool; }
+template <> inline FName FMeshElementAttributeValue::GetValue<FName>() const { check( Type == EMeshElementAttributeType::FName ); return Value_FName; }
+
+
 
 USTRUCT( BlueprintType )
 struct FMeshElementAttributeData
diff --git a/Engine/Source/Runtime/MeshDescription/Public/MeshAttributeArray.h b/Engine/Source/Runtime/MeshDescription/Public/MeshAttributeArray.h
--- a/Engine/Source/Runtime/MeshDescription/Public/MeshAttributeArray.h
+++ b/Engine/Source/Runtime/MeshDescription/Public/MeshAttributeArray.h
@@ -657,21 +657,21 @@ public:
 	template <typename AttributeType>
 	void RegisterAttribute( const FName AttributeName, const int32 NumberOfIndices = 1, const AttributeType& Default = AttributeType(), const EMeshAttributeFlags Flags = EMeshAttributeFlags::None )
 	{
-		Container.Get<TTupleIndex<AttributeType, AttributeTypes>::Value>().RegisterAttribute( AttributeName, NumberOfIndices, Default, Flags );
+		Container.template Get<TTupleIndex<AttributeType, AttributeTypes>::Value>().RegisterAttribute( AttributeName, NumberOfIndices, Default, Flags );
 	}
 
 	/** Unregister an attribute name with the given type */
 	template <typename AttributeType>
 	void UnregisterAttribute( const FName AttributeName )
 	{
-		Container.Get<TTupleIndex<AttributeType, AttributeTypes>::Value>().UnregisterAttribute( AttributeName );
+		Container.template Get<TTupleIndex<AttributeType, AttributeTypes>::Value>().UnregisterAttribute( AttributeName );
 	}
 
 	/** Determines whether an attribute of the given type exists with the given name */
 	template <typename AttributeType>
 	bool HasAttribute( const FName AttributeName ) const
 	{
-		return Container.Get<TTupleIndex<AttributeType, AttributeTypes>::Value>().HasAttribute( AttributeName );
+		return Container.template Get<TTupleIndex<AttributeType, AttributeTypes>::Value>().HasAttribute( AttributeName );
 	}
 
 	/**
@@ -690,14 +690,14 @@ public:
 	TMeshAttributeArray<AttributeType, ElementIDType>& GetAttributes( const FName AttributeName, const int32 AttributeIndex = 0 )
 	{
 		// @todo mesh description: should this handle non-existent attribute names and indices gracefully?
-		return Container.Get<TTupleIndex<AttributeType, AttributeTypes>::Value>().GetAttributes( AttributeName, AttributeIndex );
+		return Container.template Get<TTupleIndex<AttributeType, AttributeTypes>::Value>().GetAttributes( AttributeName, AttributeIndex );
 	}
 
 	template <typename AttributeType>
 	const TMeshAttributeArray<AttributeType, ElementIDType>& GetAttributes( const FName AttributeName, const int32 AttributeIndex = 0 ) const
 	{
 		// @todo mesh description: should this handle non-existent attribute names and indices gracefully?
-		return Container.Get<TTupleIndex<AttributeType, AttributeTypes>::Value>().GetAttributes( AttributeName, AttributeIndex );
+		return Container.template Get<TTupleIndex<AttributeType, AttributeTypes>::Value>().GetAttributes( AttributeName, AttributeIndex );
 	}
 
 	/**
@@ -717,14 +717,14 @@ public:
 	TAttributeIndicesArray<AttributeType, ElementIDType>& GetAttributesSet( const FName AttributeName )
 	{
 		// @todo mesh description: should this handle non-existent attribute names gracefully?
-		return Container.Get<TTupleIndex<AttributeType, AttributeTypes>::Value>().GetAttributesSet( AttributeName );
+		return Container.template Get<TTupleIndex<AttributeType, AttributeTypes>::Value>().GetAttributesSet( AttributeName );
 	}
 
 	template <typename AttributeType>
 	const TAttributeIndicesArray<AttributeType, ElementIDType>& GetAttributesSet( const FName AttributeName ) const
 	{
 		// @todo mesh description: should this handle non-existent attribute names gracefully?
-		return Container.Get<TTupleIndex<AttributeType, AttributeTypes>::Value>().GetAttributesSet( AttributeName );
+		return Container.template Get<TTupleIndex<AttributeType, AttributeTypes>::Value>().GetAttributesSet( AttributeName );
 	}
 
 	/** Returns the number of indices for the attribute with the given name */
@@ -732,33 +732,33 @@ public:
 	int32 GetAttributeIndexCount( const FName AttributeName ) const
 	{
 		// @todo mesh description: should this handle non-existent attribute names and indices gracefully?
-		return Container.Get<TTupleIndex<AttributeType, AttributeTypes>::Value>().GetAttributeIndexCount( AttributeName );
+		return Container.template Get<TTupleIndex<AttributeType, AttributeTypes>::Value>().GetAttributeIndexCount( AttributeName );
 	}
 
 	/** Sets the number of indices for the attribute with the given name */
 	template <typename AttributeType>
 	void SetAttributeIndexCount( const FName AttributeName, const int32 NumIndices )
 	{
-		Container.Get<TTupleIndex<AttributeType, AttributeTypes>::Value>().SetAttributeIndexCount( AttributeName, NumIndices );
+		Container.template Get<TTupleIndex<AttributeType, AttributeTypes>::Value>().SetAttributeIndexCount( AttributeName, NumIndices );
 	}
 
 	/** Returns an array of all the attribute names registered for this attribute type */
 	template <typename AttributeType, typename Allocator>
 	void GetAttributeNames( TArray<FName, Allocator>& OutAttributeNames ) const
 	{
-		Container.Get<TTupleIndex<AttributeType, AttributeTypes>::Value>().GetAttributeNames( OutAttributeNames );
+		Container.template Get<TTupleIndex<AttributeType, AttributeTypes>::Value>().GetAttributeNames( OutAttributeNames );
 	}
 
 	template <typename AttributeType>
 	AttributeType GetAttribute( const ElementIDType ElementID, const FName AttributeName, const int32 AttributeIndex = 0 ) const
 	{
-		return Container.Get<TTupleIndex<AttributeType, AttributeTypes>::Value>().GetAttribute( ElementID, AttributeName, AttributeIndex );
+		return Container.template Get<TTupleIndex<AttributeType, AttributeTypes>::Value>().GetAttribute( ElementID, AttributeName, AttributeIndex );
 	}
 
 	template <typename AttributeType>
 	void SetAttribute( const ElementIDType ElementID, const FName AttributeName, const int32 AttributeIndex, const AttributeType& AttributeValue )
 	{
-		Container.Get<TTupleIndex<AttributeType, AttributeTypes>::Value>().SetAttribute( ElementID, AttributeName, AttributeIndex, AttributeValue );
+		Container.template Get<TTupleIndex<AttributeType, AttributeTypes>::Value>().SetAttribute( ElementID, AttributeName, AttributeIndex, AttributeValue );
 	}
 
 	/** Inserts a default-initialized value for all attributes of the given ID */
