Only issue the #pragma to disable warning for long class names on MSVC 2015. This fixes warnings on iOS.

diff --git a/Engine/Source/Runtime/MeshDescription/Public/MeshDescription.h b/Engine/Source/Runtime/MeshDescription/Public/MeshDescription.h
--- a/Engine/Source/Runtime/MeshDescription/Public/MeshDescription.h
+++ b/Engine/Source/Runtime/MeshDescription/Public/MeshDescription.h
@@ -1,7 +1,9 @@
 // Copyright 1998-2017 Epic Games, Inc. All Rights Reserved.
 
 #pragma once
+#if defined(_MSC_VER) && _MSC_FULL_VER <= 190023918
 #pragma warning(disable:4503)
+#endif
 
 #include "CoreMinimal.h"
 #include "CoreTypes.h"
