Move Normals tangents creation in MeshDescription module, this will allow to share the code between editable mesh and staticmesh builder

diff --git a/Engine/Source/Developer/MeshBuilder/Private/MeshDescriptionHelper.cpp b/Engine/Source/Developer/MeshBuilder/Private/MeshDescriptionHelper.cpp
--- a/Engine/Source/Developer/MeshBuilder/Private/MeshDescriptionHelper.cpp
+++ b/Engine/Source/Developer/MeshBuilder/Private/MeshDescriptionHelper.cpp
@@ -67,7 +67,7 @@ UMeshDescription* FMeshDescriptionHelper::GetRenderMeshDescription(UObject* Owne
 		}
 
 		//This function make sure the Polygon NTB are compute and also remove degenerated triangle from the render mesh description.
-		FMeshDescriptionOperations::CreatePolygonNTB(RenderMeshDescription, BuildSettings->bRemoveDegenerates ? SMALL_NUMBER : 0.0f);
+		RenderMeshDescription->ComputePolygonTangentsAndNormals(BuildSettings->bRemoveDegenerates ? SMALL_NUMBER : 0.0f);
 		
 		//Keep the original mesh description NTBs if we do not rebuild the normals or tangents.
 		bool bComputeTangentLegacy = !BuildSettings->bUseMikkTSpace && (BuildSettings->bRecomputeNormals || BuildSettings->bRecomputeTangents);
@@ -98,14 +98,16 @@ UMeshDescription* FMeshDescriptionHelper::GetRenderMeshDescription(UObject* Owne
 		{
 			if (!bHasAllNormals)
 			{
-				FMeshDescriptionOperations::CreateNormals(RenderMeshDescription, (FMeshDescriptionOperations::ETangentOptions)TangentOptions, false);
+				//FMeshDescriptionOperations::CreateNormals(RenderMeshDescription, (FMeshDescriptionOperations::ETangentOptions)TangentOptions, false);
+				RenderMeshDescription->ComputeTangentsAndNormals(false, false);
 			}
 			FMeshDescriptionOperations::CreateMikktTangents(RenderMeshDescription, (FMeshDescriptionOperations::ETangentOptions)TangentOptions);
 		}
 		else if(!bHasAllNormals || !bHasAllTangents)
 		{
+			RenderMeshDescription->ComputeTangentsAndNormals(true, false);
 			//Set the compute tangent to true when we do not build using mikkt space
-			FMeshDescriptionOperations::CreateNormals(RenderMeshDescription, (FMeshDescriptionOperations::ETangentOptions)TangentOptions, true);
+			//FMeshDescriptionOperations::CreateNormals(RenderMeshDescription, (FMeshDescriptionOperations::ETangentOptions)TangentOptions, true);
 		}
 	}
 
diff --git a/Engine/Source/Developer/MeshDescriptionOperations/Private/MeshDescriptionOperations.cpp b/Engine/Source/Developer/MeshDescriptionOperations/Private/MeshDescriptionOperations.cpp
--- a/Engine/Source/Developer/MeshDescriptionOperations/Private/MeshDescriptionOperations.cpp
+++ b/Engine/Source/Developer/MeshDescriptionOperations/Private/MeshDescriptionOperations.cpp
@@ -447,15 +447,17 @@ void FMeshDescriptionOperations::ConverFromRawMesh(const struct FRawMesh &Source
 			NewTriangle.SetVertexInstanceID(Corner, VertexInstanceID);
 		}
 	}
-	CreatePolygonNTB(DestinationMeshDescription, 0.0f);
+	DestinationMeshDescription->ComputePolygonTangentsAndNormals(0.0f);
 
-	if (!bHasNormals || !bHasTangents)
+	//Create the missing normals and tangents
+	if (!bHasNormals)
+	{
+		//CreateNormals(DestinationMeshDescription, ETangentOptions::BlendOverlappingNormals, false);
+		DestinationMeshDescription->ComputeTangentsAndNormals(false, false);
+	}
+
+	if (!bHasTangents)
 	{
-		//Create the missing normals and tangents
-		if (!bHasNormals)
-		{
-			CreateNormals(DestinationMeshDescription, ETangentOptions::BlendOverlappingNormals, false);
-		}
 		CreateMikktTangents(DestinationMeshDescription, ETangentOptions::BlendOverlappingNormals);
 	}
 
diff --git a/Engine/Source/Runtime/MeshDescription/Private/MeshDescription.cpp b/Engine/Source/Runtime/MeshDescription/Private/MeshDescription.cpp
--- a/Engine/Source/Runtime/MeshDescription/Private/MeshDescription.cpp
+++ b/Engine/Source/Runtime/MeshDescription/Private/MeshDescription.cpp
@@ -608,6 +608,379 @@ void UMeshDescription::TriangulateMesh()
 	}
 }
 
+bool UMeshDescription::ComputePolygonTangentsAndNormals(const FPolygonID PolygonID
+	, float ComparisonThreshold
+	, const TVertexAttributeArray<FVector>& VertexPositions
+	, const TVertexInstanceAttributeArray<FVector2D>& VertexUVs
+	, TPolygonAttributeArray<FVector>& PolygonNormals
+	, TPolygonAttributeArray<FVector>& PolygonTangents
+	, TPolygonAttributeArray<FVector>& PolygonBinormals
+	, TPolygonAttributeArray<FVector>& PolygonCenters)
+{
+	bool bValidNTBs = true;
+	// Calculate the center of this polygon
+	FVector Center = FVector::ZeroVector;
+	const TArray<FVertexInstanceID>& VertexInstanceIDs = GetPolygonPerimeterVertexInstances(PolygonID);
+	for (const FVertexInstanceID VertexInstanceID : VertexInstanceIDs)
+	{
+		Center += VertexPositions[GetVertexInstanceVertex(VertexInstanceID)];
+	}
+	Center /= float(VertexInstanceIDs.Num());
+
+	// Calculate the tangent basis for the polygon, based on the average of all constituent triangles
+	FVector Normal = FVector::ZeroVector;
+	FVector Tangent = FVector::ZeroVector;
+	FVector Binormal = FVector::ZeroVector;
+
+	for (const FMeshTriangle& Triangle : GetPolygonTriangles(PolygonID))
+	{
+		const FVertexID VertexID0 = GetVertexInstanceVertex(Triangle.VertexInstanceID0);
+		const FVertexID VertexID1 = GetVertexInstanceVertex(Triangle.VertexInstanceID1);
+		const FVertexID VertexID2 = GetVertexInstanceVertex(Triangle.VertexInstanceID2);
+
+		const FVector DPosition1 = VertexPositions[VertexID1] - VertexPositions[VertexID0];
+		const FVector DPosition2 = VertexPositions[VertexID2] - VertexPositions[VertexID0];
+
+		const FVector2D DUV1 = VertexUVs[Triangle.VertexInstanceID1] - VertexUVs[Triangle.VertexInstanceID0];
+		const FVector2D DUV2 = VertexUVs[Triangle.VertexInstanceID2] - VertexUVs[Triangle.VertexInstanceID0];
+
+		// We have a left-handed coordinate system, but a counter-clockwise winding order
+		// Hence normal calculation has to take the triangle vectors cross product in reverse.
+		FVector TmpNormal = FVector::CrossProduct(DPosition2, DPosition1);
+		if (!TmpNormal.IsNearlyZero(ComparisonThreshold) && !TmpNormal.ContainsNaN())
+		{
+			Normal += TmpNormal;
+			// ...and tangent space seems to be right-handed.
+			const float DetUV = FVector2D::CrossProduct(DUV1, DUV2);
+			const float InvDetUV = (DetUV == 0.0f) ? 0.0f : 1.0f / DetUV;
+
+			Tangent += (DPosition1 * DUV2.Y - DPosition2 * DUV1.Y) * InvDetUV;
+			Binormal += (DPosition2 * DUV1.X - DPosition1 * DUV2.X) * InvDetUV;
+		}
+		else
+		{
+			//The polygon is degenerated
+			bValidNTBs = false;
+		}
+	}
+
+	PolygonNormals[PolygonID] = Normal.GetSafeNormal();
+	PolygonTangents[PolygonID] = Tangent.GetSafeNormal();
+	PolygonBinormals[PolygonID] = Binormal.GetSafeNormal();
+	PolygonCenters[PolygonID] = Center;
+	
+	return bValidNTBs;
+}
+
+void UMeshDescription::ComputePolygonTangentsAndNormals(const TArray<FPolygonID>& PolygonIDs, float ComparisonThreshold)
+{
+	const TVertexAttributeArray<FVector>& VertexPositions = VertexAttributes().GetAttributes<FVector>(MeshAttribute::Vertex::Position);
+	const TVertexInstanceAttributeArray<FVector2D>& VertexUVs = VertexInstanceAttributes().GetAttributes<FVector2D>(MeshAttribute::VertexInstance::TextureCoordinate, 0);
+	TPolygonAttributeArray<FVector>& PolygonNormals = PolygonAttributes().GetAttributes<FVector>(MeshAttribute::Polygon::Normal);
+	TPolygonAttributeArray<FVector>& PolygonTangents = PolygonAttributes().GetAttributes<FVector>(MeshAttribute::Polygon::Tangent);
+	TPolygonAttributeArray<FVector>& PolygonBinormals = PolygonAttributes().GetAttributes<FVector>(MeshAttribute::Polygon::Binormal);
+	TPolygonAttributeArray<FVector>& PolygonCenters = PolygonAttributes().GetAttributes<FVector>(MeshAttribute::Polygon::Center);
+
+	TArray<FPolygonID> DegeneratePolygonIDs;
+	for (const FPolygonID PolygonID : PolygonIDs)
+	{
+		if (!ComputePolygonTangentsAndNormals(PolygonID, ComparisonThreshold, VertexPositions, VertexUVs, PolygonNormals, PolygonTangents, PolygonBinormals, PolygonCenters))
+		{
+			DegeneratePolygonIDs.Add(PolygonID);
+		}
+	}
+	//Remove degenerated polygons
+	//Delete the degenerated polygons. The array is fill only if the remove degenerated option is turn on.
+	if (DegeneratePolygonIDs.Num() > 0)
+	{
+		TArray<FEdgeID> OrphanedEdges;
+		TArray<FVertexInstanceID> OrphanedVertexInstances;
+		TArray<FPolygonGroupID> OrphanedPolygonGroups;
+		TArray<FVertexID> OrphanedVertices;
+		for (FPolygonID& PolygonID : DegeneratePolygonIDs)
+		{
+			DeletePolygon(PolygonID, &OrphanedEdges, &OrphanedVertexInstances, &OrphanedPolygonGroups);
+		}
+		for (FPolygonGroupID& PolygonGroupID : OrphanedPolygonGroups)
+		{
+			DeletePolygonGroup(PolygonGroupID);
+		}
+		for (FVertexInstanceID& VertexInstanceID : OrphanedVertexInstances)
+		{
+			DeleteVertexInstance(VertexInstanceID, &OrphanedVertices);
+		}
+		for (FEdgeID& EdgeID : OrphanedEdges)
+		{
+			DeleteEdge(EdgeID, &OrphanedVertices);
+		}
+		for (FVertexID& VertexID : OrphanedVertices)
+		{
+			DeleteVertex(VertexID);
+		}
+		//Compact and Remap IDs so we have clean ID from 0 to n since we just erase some polygons
+		//The render build need to have compact ID
+		FElementIDRemappings RemappingInfos;
+		Compact(RemappingInfos);
+	}
+}
+
+void UMeshDescription::ComputePolygonTangentsAndNormals(float ComparisonThreshold)
+{
+	TArray<FPolygonID> PolygonsToComputeNTBs;
+	PolygonsToComputeNTBs.Reserve(Polygons().Num());
+	for (const FPolygonID& PolygonID : Polygons().GetElementIDs())
+	{
+		PolygonsToComputeNTBs.Add(PolygonID);
+	}
+	ComputePolygonTangentsAndNormals(PolygonsToComputeNTBs, ComparisonThreshold);
+}
+
+void UMeshDescription::GetConnectedSoftEdges(const FVertexID VertexID, TArray<FEdgeID>& OutConnectedSoftEdges) const
+{
+	OutConnectedSoftEdges.Reset();
+
+	const TEdgeAttributeArray<bool>& EdgeHardnesses = EdgeAttributes().GetAttributes<bool>(MeshAttribute::Edge::IsHard);
+	for (const FEdgeID ConnectedEdgeID : GetVertex(VertexID).ConnectedEdgeIDs)
+	{
+		if (!EdgeHardnesses[ConnectedEdgeID])
+		{
+			OutConnectedSoftEdges.Add(ConnectedEdgeID);
+		}
+	}
+}
+
+void UMeshDescription::GetPolygonsInSameSoftEdgedGroupAsPolygon(const FPolygonID PolygonID, const TArray<FPolygonID>& CandidatePolygonIDs, const TArray<FEdgeID>& SoftEdgeIDs, TArray<FPolygonID>& OutPolygonIDs) const
+{
+	// The aim of this method is:
+	// - given a polygon ID,
+	// - given a set of candidate polygons connected to the same vertex (which should include the polygon ID),
+	// - given a set of soft edges connected to the same vertex,
+	// return the polygon IDs which form an adjacent run without crossing a hard edge.
+
+	OutPolygonIDs.Reset();
+
+	// Maintain a list of polygon IDs to be examined. Adjacents are added to the list if suitable.
+	// Add the start poly here.
+	static TArray<FPolygonID> PolygonsToCheck;
+	PolygonsToCheck.Reset(CandidatePolygonIDs.Num());
+	PolygonsToCheck.Add(PolygonID);
+
+	int32 Index = 0;
+	while (Index < PolygonsToCheck.Num())
+	{
+		const FPolygonID PolygonToCheck = PolygonsToCheck[Index];
+		Index++;
+
+		if (CandidatePolygonIDs.Contains(PolygonToCheck))
+		{
+			OutPolygonIDs.Add(PolygonToCheck);
+
+			// Now look at its adjacent polygons. If they are joined by a soft edge which includes the vertex we're interested in, we want to consider them.
+			// We take a shortcut by doing this process in reverse: we already know all the soft edges we are interested in, so check if any of them
+			// have the current polygon as an adjacent.
+			for (const FEdgeID SoftEdgeID : SoftEdgeIDs)
+			{
+				const TArray<FPolygonID>& EdgeConnectedPolygons = GetEdgeConnectedPolygons(SoftEdgeID);
+				if (EdgeConnectedPolygons.Contains(PolygonToCheck))
+				{
+					for (const FPolygonID AdjacentPolygon : EdgeConnectedPolygons)
+					{
+						// Only add new polygons which haven't yet been added to the list. This prevents circular runs of polygons triggering infinite loops.
+						PolygonsToCheck.AddUnique(AdjacentPolygon);
+					}
+				}
+			}
+		}
+	}
+}
+
+void UMeshDescription::GetVertexConnectedPolygonsInSameSoftEdgedGroup(const FVertexID VertexID, const FPolygonID PolygonID, TArray<FPolygonID>& OutPolygonIDs) const
+{
+	// The aim here is to determine which polygons form part of the same soft edged group as the polygons attached to this vertex.
+	// They should all contribute to the final vertex instance normal.
+
+	// Get all polygons connected to this vertex.
+	static TArray<FPolygonID> ConnectedPolygons;
+	GetVertexConnectedPolygons(VertexID, ConnectedPolygons);
+
+	// Cache a list of all soft edges which share this vertex.
+	// We're only interested in finding adjacent polygons which are not the other side of a hard edge.
+	static TArray<FEdgeID> ConnectedSoftEdges;
+	GetConnectedSoftEdges(VertexID, ConnectedSoftEdges);
+
+	GetPolygonsInSameSoftEdgedGroupAsPolygon(PolygonID, ConnectedPolygons, ConnectedSoftEdges, OutPolygonIDs);
+}
+
+float UMeshDescription::GetPolygonCornerAngleForVertex(const FPolygonID PolygonID, const FVertexID VertexID) const
+{
+	const FMeshPolygon& Polygon = GetPolygon(PolygonID);
+
+	// Lambda function which returns the inner angle at a given index on a polygon contour
+	auto GetContourAngle = [this](const FMeshPolygonContour& Contour, const int32 ContourIndex)
+	{
+		const int32 NumVertices = Contour.VertexInstanceIDs.Num();
+
+		const int32 PrevIndex = (ContourIndex + NumVertices - 1) % NumVertices;
+		const int32 NextIndex = (ContourIndex + 1) % NumVertices;
+
+		const FVertexID PrevVertexID = GetVertexInstanceVertex(Contour.VertexInstanceIDs[PrevIndex]);
+		const FVertexID ThisVertexID = GetVertexInstanceVertex(Contour.VertexInstanceIDs[ContourIndex]);
+		const FVertexID NextVertexID = GetVertexInstanceVertex(Contour.VertexInstanceIDs[NextIndex]);
+
+		const TVertexAttributeArray<FVector>& VertexPositions = VertexAttributes().GetAttributes<FVector>(MeshAttribute::Vertex::Position);
+
+		const FVector PrevVertexPosition = VertexPositions[PrevVertexID];
+		const FVector ThisVertexPosition = VertexPositions[ThisVertexID];
+		const FVector NextVertexPosition = VertexPositions[NextVertexID];
+
+		const FVector Direction1 = (PrevVertexPosition - ThisVertexPosition).GetSafeNormal();
+		const FVector Direction2 = (NextVertexPosition - ThisVertexPosition).GetSafeNormal();
+
+		return FMath::Acos(FVector::DotProduct(Direction1, Direction2));
+	};
+
+	const FVertexInstanceArray& VertexInstancesRef = VertexInstances();
+	auto IsVertexInstancedFromThisVertex = [&VertexInstancesRef, VertexID](const FVertexInstanceID VertexInstanceID)
+	{
+		return VertexInstancesRef[VertexInstanceID].VertexID == VertexID;
+	};
+
+	// First look for the vertex instance in the perimeter
+	int32 ContourIndex = Polygon.PerimeterContour.VertexInstanceIDs.IndexOfByPredicate(IsVertexInstancedFromThisVertex);
+	if (ContourIndex != INDEX_NONE)
+	{
+		// Return the internal angle if found
+		return GetContourAngle(Polygon.PerimeterContour, ContourIndex);
+	}
+	else
+	{
+		// If not found, look in all the holes
+		for (const FMeshPolygonContour& HoleContour : Polygon.HoleContours)
+		{
+			ContourIndex = HoleContour.VertexInstanceIDs.IndexOfByPredicate(IsVertexInstancedFromThisVertex);
+			if (ContourIndex != INDEX_NONE)
+			{
+				// Hole vertex contribution is the part which ISN'T the internal angle of the contour, so subtract from 2*pi
+				return (2.0f * PI) - GetContourAngle(HoleContour, ContourIndex);
+			}
+		}
+	}
+
+	// Found nothing; return 0
+	return 0.0f;
+}
+
+void UMeshDescription::ComputeTangentsAndNormals(const FVertexInstanceID& VertexInstanceID
+	, bool bComputeTangents
+	, bool bUseWeightedNormals
+	, const TPolygonAttributeArray<FVector>& PolygonNormals
+	, const TPolygonAttributeArray<FVector>& PolygonTangents
+	, const TPolygonAttributeArray<FVector>& PolygonBinormals
+	, TVertexInstanceAttributeArray<FVector>& VertexNormals
+	, TVertexInstanceAttributeArray<FVector>& VertexTangents
+	, TVertexInstanceAttributeArray<float>& VertexBinormalSigns)
+{
+	FVector Normal = FVector::ZeroVector;
+	FVector Tangent = FVector::ZeroVector;
+	FVector Binormal = FVector::ZeroVector;
+
+	FVector& NormalRef = VertexNormals[VertexInstanceID];
+	FVector& TangentRef = VertexTangents[VertexInstanceID];
+	float& BinormalRef = VertexBinormalSigns[VertexInstanceID];
+
+	if (!NormalRef.IsNearlyZero() && !TangentRef.IsNearlyZero() && !FMath::IsNearlyZero(BinormalRef))
+	{
+		//Nothing to compute
+		return;
+	}
+
+	const FVertexID VertexID = GetVertexInstanceVertex(VertexInstanceID);
+
+	// Get all polygons connected to this vertex instance, and also any in the same smoothing group connected to a different vertex instance
+	// (as they still have influence over the normal).
+	static TArray<FPolygonID> AllConnectedPolygons;
+	const TArray<FPolygonID>& VertexInstanceConnectedPolygons = GetVertexInstanceConnectedPolygons(VertexInstanceID);
+	check(VertexInstanceConnectedPolygons.Num() > 0);
+	GetVertexConnectedPolygonsInSameSoftEdgedGroup(VertexID, VertexInstanceConnectedPolygons[0], AllConnectedPolygons);
+
+	// The vertex instance normal is computed as a sum of all connected polygons' normals, weighted by the angle they make with the vertex
+	for (const FPolygonID ConnectedPolygonID : AllConnectedPolygons)
+	{
+		const float Angle = bUseWeightedNormals ? GetPolygonCornerAngleForVertex(ConnectedPolygonID, VertexID) : 1.0f;
+
+		Normal += PolygonNormals[ConnectedPolygonID] * Angle;
+
+		// If this polygon is actually connected to the vertex instance we're processing, also include its contributions towards the tangent
+		if (VertexInstanceConnectedPolygons.Contains(ConnectedPolygonID))
+		{
+			Tangent += PolygonTangents[ConnectedPolygonID] * Angle;
+			Binormal += PolygonBinormals[ConnectedPolygonID] * Angle;
+		}
+	}
+
+	// Normalize Normal
+	Normal = Normal.GetSafeNormal();
+	float BinormalSign = 1.0f;
+	if (bComputeTangents)
+	{
+		// Make Tangent orthonormal to Normal.
+		// This is a quicker method than normalizing Tangent, taking the cross product Normal X Tangent, and then a further cross product with that result
+		Tangent = (Tangent - Normal * FVector::DotProduct(Normal, Tangent)).GetSafeNormal();
+
+		// Calculate binormal sign
+		BinormalSign = (FVector::DotProduct(FVector::CrossProduct(Normal, Tangent), Binormal) < 0.0f) ? -1.0f : 1.0f;
+	}
+
+	//Set the value that need to be set
+	if (NormalRef.IsNearlyZero())
+	{
+		NormalRef = Normal;
+	}
+	if (bComputeTangents)
+	{
+		if (TangentRef.IsNearlyZero())
+		{
+			TangentRef = Tangent;
+		}
+		if (FMath::IsNearlyZero(BinormalRef))
+		{
+			BinormalRef = BinormalSign;
+		}
+	}
+}
+
+void UMeshDescription::ComputeTangentsAndNormals(const TArray<FVertexInstanceID>& VertexInstanceIDs, bool bComputeTangents, bool bUseWeightedNormals)
+{
+	const TPolygonAttributeArray<FVector>& PolygonNormals = PolygonAttributes().GetAttributes<FVector>(MeshAttribute::Polygon::Normal);
+	const TPolygonAttributeArray<FVector>& PolygonTangents = PolygonAttributes().GetAttributes<FVector>(MeshAttribute::Polygon::Tangent);
+	const TPolygonAttributeArray<FVector>& PolygonBinormals = PolygonAttributes().GetAttributes<FVector>(MeshAttribute::Polygon::Binormal);
+
+	TVertexInstanceAttributeArray<FVector>& VertexNormals = VertexInstanceAttributes().GetAttributes<FVector>(MeshAttribute::VertexInstance::Normal, 0);
+	TVertexInstanceAttributeArray<FVector>& VertexTangents = VertexInstanceAttributes().GetAttributes<FVector>(MeshAttribute::VertexInstance::Tangent, 0);
+	TVertexInstanceAttributeArray<float>& VertexBinormalSigns = VertexInstanceAttributes().GetAttributes<float>(MeshAttribute::VertexInstance::BinormalSign, 0);
+
+	for (const FVertexInstanceID VertexInstanceID : VertexInstanceIDs)
+	{
+		ComputeTangentsAndNormals(VertexInstanceID, bComputeTangents, bUseWeightedNormals, PolygonNormals, PolygonTangents, PolygonBinormals, VertexNormals, VertexTangents, VertexBinormalSigns);
+	}
+}
+
+void UMeshDescription::ComputeTangentsAndNormals(bool bComputeTangents, bool bUseWeightedNormals)
+{
+	const TPolygonAttributeArray<FVector>& PolygonNormals = PolygonAttributes().GetAttributes<FVector>(MeshAttribute::Polygon::Normal);
+	const TPolygonAttributeArray<FVector>& PolygonTangents = PolygonAttributes().GetAttributes<FVector>(MeshAttribute::Polygon::Tangent);
+	const TPolygonAttributeArray<FVector>& PolygonBinormals = PolygonAttributes().GetAttributes<FVector>(MeshAttribute::Polygon::Binormal);
+
+	TVertexInstanceAttributeArray<FVector>& VertexNormals = VertexInstanceAttributes().GetAttributes<FVector>(MeshAttribute::VertexInstance::Normal, 0);
+	TVertexInstanceAttributeArray<FVector>& VertexTangents = VertexInstanceAttributes().GetAttributes<FVector>(MeshAttribute::VertexInstance::Tangent, 0);
+	TVertexInstanceAttributeArray<float>& VertexBinormalSigns = VertexInstanceAttributes().GetAttributes<float>(MeshAttribute::VertexInstance::BinormalSign, 0);
+
+	for (const FVertexInstanceID VertexInstanceID : VertexInstances().GetElementIDs())
+	{
+		ComputeTangentsAndNormals(VertexInstanceID, bComputeTangents, bUseWeightedNormals, PolygonNormals, PolygonTangents, PolygonBinormals, VertexNormals, VertexTangents, VertexBinormalSigns);
+	}
+}
+
 void UMeshDescription::ReversePolygonFacing(const FPolygonID PolygonID)
 {
 	FMeshPolygon& Polygon = GetPolygon(PolygonID);
diff --git a/Engine/Source/Runtime/MeshDescription/Public/MeshDescription.h b/Engine/Source/Runtime/MeshDescription/Public/MeshDescription.h
--- a/Engine/Source/Runtime/MeshDescription/Public/MeshDescription.h
+++ b/Engine/Source/Runtime/MeshDescription/Public/MeshDescription.h
@@ -799,6 +799,16 @@ public:
 	void ComputePolygonTriangulation(const FPolygonID PolygonID, TArray<FMeshTriangle>& OutTriangles);
 	void TriangulateMesh();
 
+	/** Set the polygon tangent and normal only for the specified polygonIDs */
+	void ComputePolygonTangentsAndNormals(const TArray<FPolygonID>& PolygonIDs, float ComparisonThreshold = 0.0f);
+	/** Set the polygon tangent and normal for all polygons in the mesh description. */
+	void ComputePolygonTangentsAndNormals(float ComparisonThreshold = 0.0f);
+	
+	/** Set the vertex instance tangent and normal only for the specified VertexInstanceIDs */
+	void ComputeTangentsAndNormals(const TArray<FVertexInstanceID>& VertexInstanceIDs, bool bComputeTangents, bool bUseWeightedNormals);
+	/** Set the vertex instance tangent and normal for all vertex instances in the mesh description. */
+	void ComputeTangentsAndNormals(bool bComputeTangents, bool bUseWeightedNormals);
+
 	void ReversePolygonFacing(const FPolygonID PolygonID);
 	void ReverseAllPolygonFacing();
 
@@ -808,6 +818,32 @@ private:
 	FPlane ComputePolygonPlane(const FPolygonID PolygonID) const;
 	FVector ComputePolygonNormal(const FPolygonID PolygonID) const;
 
+	float GetPolygonCornerAngleForVertex(const FPolygonID PolygonID, const FVertexID VertexID) const;
+	bool ComputePolygonTangentsAndNormals(
+		  const FPolygonID PolygonID
+		, float ComparisonThreshold
+		, const TVertexAttributeArray<FVector>& VertexPositions
+		, const TVertexInstanceAttributeArray<FVector2D>& VertexUVs
+		, TPolygonAttributeArray<FVector>& PolygonNormals
+		, TPolygonAttributeArray<FVector>& PolygonTangents
+		, TPolygonAttributeArray<FVector>& PolygonBinormals
+		, TPolygonAttributeArray<FVector>& PolygonCenters
+	);
+
+	void GetVertexConnectedPolygonsInSameSoftEdgedGroup(const FVertexID VertexID, const FPolygonID PolygonID, TArray<FPolygonID>& OutPolygonIDs) const;
+	void GetPolygonsInSameSoftEdgedGroupAsPolygon(const FPolygonID PolygonID, const TArray<FPolygonID>& CandidatePolygonIDs, const TArray<FEdgeID>& SoftEdgeIDs, TArray<FPolygonID>& OutPolygonIDs) const;
+	void GetConnectedSoftEdges(const FVertexID VertexID, TArray<FEdgeID>& OutConnectedSoftEdges) const;
+	void ComputeTangentsAndNormals(
+		  const FVertexInstanceID& VertexInstanceID
+		, bool bComputeTangents
+		, bool bUseWeightedNormals
+		, const TPolygonAttributeArray<FVector>& PolygonNormals
+		, const TPolygonAttributeArray<FVector>& PolygonTangents
+		, const TPolygonAttributeArray<FVector>& PolygonBinormals
+		, TVertexInstanceAttributeArray<FVector>& VertexNormals
+		, TVertexInstanceAttributeArray<FVector>& VertexTangents
+		, TVertexInstanceAttributeArray<float>& VertexBinormalSigns
+	);
 private:
 
 	/** Given a set of index remappings, fixes up references to element IDs */
