When StaticMesh build options are set to use the imported normals, we now use the imported normals to compute the tangents when using the builtin option.

diff --git a/Engine/Source/Developer/MeshBuilder/Private/MeshDescriptionHelper.cpp b/Engine/Source/Developer/MeshBuilder/Private/MeshDescriptionHelper.cpp
--- a/Engine/Source/Developer/MeshBuilder/Private/MeshDescriptionHelper.cpp
+++ b/Engine/Source/Developer/MeshBuilder/Private/MeshDescriptionHelper.cpp
@@ -98,16 +98,15 @@ UMeshDescription* FMeshDescriptionHelper::GetRenderMeshDescription(UObject* Owne
 		{
 			if (!bHasAllNormals)
 			{
-				//FMeshDescriptionOperations::CreateNormals(RenderMeshDescription, (FMeshDescriptionOperations::ETangentOptions)TangentOptions, false);
-				RenderMeshDescription->ComputeTangentsAndNormals(false, false);
+				EComputeNTBsOptions ComputeNTBsOptions = EComputeNTBsOptions::Normals;
+				RenderMeshDescription->ComputeTangentsAndNormals(ComputeNTBsOptions);
 			}
 			FMeshDescriptionOperations::CreateMikktTangents(RenderMeshDescription, (FMeshDescriptionOperations::ETangentOptions)TangentOptions);
 		}
 		else if(!bHasAllNormals || !bHasAllTangents)
 		{
-			RenderMeshDescription->ComputeTangentsAndNormals(true, false);
-			//Set the compute tangent to true when we do not build using mikkt space
-			//FMeshDescriptionOperations::CreateNormals(RenderMeshDescription, (FMeshDescriptionOperations::ETangentOptions)TangentOptions, true);
+			EComputeNTBsOptions ComputeNTBsOptions = (bHasAllNormals ? EComputeNTBsOptions::None : EComputeNTBsOptions::Normals) | (bHasAllTangents ? EComputeNTBsOptions::None : EComputeNTBsOptions::Tangents);
+			RenderMeshDescription->ComputeTangentsAndNormals(ComputeNTBsOptions);
 		}
 	}
 
diff --git a/Engine/Source/Developer/MeshDescriptionOperations/Private/MeshDescriptionOperations.cpp b/Engine/Source/Developer/MeshDescriptionOperations/Private/MeshDescriptionOperations.cpp
--- a/Engine/Source/Developer/MeshDescriptionOperations/Private/MeshDescriptionOperations.cpp
+++ b/Engine/Source/Developer/MeshDescriptionOperations/Private/MeshDescriptionOperations.cpp
@@ -449,16 +449,11 @@ void FMeshDescriptionOperations::ConverFromRawMesh(const struct FRawMesh &Source
 	}
 	DestinationMeshDescription->ComputePolygonTangentsAndNormals(0.0f);
 
-	//Create the missing normals and tangents
-	if (!bHasNormals)
+	//Create the missing normals and tangents, should we use Mikkt space for tangent???
+	if (!bHasNormals || !bHasTangents)
 	{
-		//CreateNormals(DestinationMeshDescription, ETangentOptions::BlendOverlappingNormals, false);
-		DestinationMeshDescription->ComputeTangentsAndNormals(false, false);
-	}
-
-	if (!bHasTangents)
-	{
-		CreateMikktTangents(DestinationMeshDescription, ETangentOptions::BlendOverlappingNormals);
+		EComputeNTBsOptions ComputeNTBsOptions = (bHasNormals ? EComputeNTBsOptions::None : EComputeNTBsOptions::Normals) | (bHasTangents ? EComputeNTBsOptions::None : EComputeNTBsOptions::Tangents);
+		DestinationMeshDescription->ComputeTangentsAndNormals(ComputeNTBsOptions);
 	}
 
 	ConvertSmoothGroupToHardEdges(SourceRawMesh, DestinationMeshDescription);
diff --git a/Engine/Source/Runtime/MeshDescription/Private/MeshDescription.cpp b/Engine/Source/Runtime/MeshDescription/Private/MeshDescription.cpp
--- a/Engine/Source/Runtime/MeshDescription/Private/MeshDescription.cpp
+++ b/Engine/Source/Runtime/MeshDescription/Private/MeshDescription.cpp
@@ -871,8 +871,7 @@ float UMeshDescription::GetPolygonCornerAngleForVertex(const FPolygonID PolygonI
 }
 
 void UMeshDescription::ComputeTangentsAndNormals(const FVertexInstanceID& VertexInstanceID
-	, bool bComputeTangents
-	, bool bUseWeightedNormals
+	, EComputeNTBsOptions ComputeNTBsOptions
 	, const TPolygonAttributeArray<FVector>& PolygonNormals
 	, const TPolygonAttributeArray<FVector>& PolygonTangents
 	, const TPolygonAttributeArray<FVector>& PolygonBinormals
@@ -880,6 +879,10 @@ void UMeshDescription::ComputeTangentsAndNormals(const FVertexInstanceID& Vertex
 	, TVertexInstanceAttributeArray<FVector>& VertexTangents
 	, TVertexInstanceAttributeArray<float>& VertexBinormalSigns)
 {
+	bool bComputeNormals = !!(ComputeNTBsOptions & EComputeNTBsOptions::Normals);
+	bool bComputeTangents = !!(ComputeNTBsOptions & EComputeNTBsOptions::Tangents);
+	bool bUseWeightedNormals = !!(ComputeNTBsOptions & EComputeNTBsOptions::WeightedNTBs);
+
 	FVector Normal = FVector::ZeroVector;
 	FVector Tangent = FVector::ZeroVector;
 	FVector Binormal = FVector::ZeroVector;
@@ -888,7 +891,7 @@ void UMeshDescription::ComputeTangentsAndNormals(const FVertexInstanceID& Vertex
 	FVector& TangentRef = VertexTangents[VertexInstanceID];
 	float& BinormalRef = VertexBinormalSigns[VertexInstanceID];
 
-	if (!NormalRef.IsNearlyZero() && !TangentRef.IsNearlyZero() && !FMath::IsNearlyZero(BinormalRef))
+	if (!bComputeNormals && !bComputeTangents)
 	{
 		//Nothing to compute
 		return;
@@ -896,30 +899,57 @@ void UMeshDescription::ComputeTangentsAndNormals(const FVertexInstanceID& Vertex
 
 	const FVertexID VertexID = GetVertexInstanceVertex(VertexInstanceID);
 
-	// Get all polygons connected to this vertex instance, and also any in the same smoothing group connected to a different vertex instance
-	// (as they still have influence over the normal).
-	static TArray<FPolygonID> AllConnectedPolygons;
-	const TArray<FPolygonID>& VertexInstanceConnectedPolygons = GetVertexInstanceConnectedPolygons(VertexInstanceID);
-	check(VertexInstanceConnectedPolygons.Num() > 0);
-	GetVertexConnectedPolygonsInSameSoftEdgedGroup(VertexID, VertexInstanceConnectedPolygons[0], AllConnectedPolygons);
-
-	// The vertex instance normal is computed as a sum of all connected polygons' normals, weighted by the angle they make with the vertex
-	for (const FPolygonID ConnectedPolygonID : AllConnectedPolygons)
+	if (bComputeNormals || NormalRef.IsNearlyZero())
 	{
-		const float Angle = bUseWeightedNormals ? GetPolygonCornerAngleForVertex(ConnectedPolygonID, VertexID) : 1.0f;
+		// Get all polygons connected to this vertex instance
+		static TArray<FPolygonID> AllConnectedPolygons;
+		const TArray<FPolygonID>& VertexInstanceConnectedPolygons = GetVertexInstanceConnectedPolygons(VertexInstanceID);
+		check(VertexInstanceConnectedPolygons.Num() > 0);
+		// Add also any in the same smoothing group connected to a different vertex instance
+		// (as they still have influence over the normal).
+		GetVertexConnectedPolygonsInSameSoftEdgedGroup(VertexID, VertexInstanceConnectedPolygons[0], AllConnectedPolygons);
+		// The vertex instance normal is computed as a sum of all connected polygons' normals, weighted by the angle they make with the vertex
+		for (const FPolygonID ConnectedPolygonID : AllConnectedPolygons)
+		{
+			const float Angle = bUseWeightedNormals ? GetPolygonCornerAngleForVertex(ConnectedPolygonID, VertexID) : 1.0f;
 
-		Normal += PolygonNormals[ConnectedPolygonID] * Angle;
+			Normal += PolygonNormals[ConnectedPolygonID] * Angle;
 
-		// If this polygon is actually connected to the vertex instance we're processing, also include its contributions towards the tangent
-		if (VertexInstanceConnectedPolygons.Contains(ConnectedPolygonID))
-		{
-			Tangent += PolygonTangents[ConnectedPolygonID] * Angle;
-			Binormal += PolygonBinormals[ConnectedPolygonID] * Angle;
+			// If this polygon is actually connected to the vertex instance we're processing, also include its contributions towards the tangent
+			if (VertexInstanceConnectedPolygons.Contains(ConnectedPolygonID))
+			{
+				Tangent += PolygonTangents[ConnectedPolygonID] * Angle;
+				Binormal += PolygonBinormals[ConnectedPolygonID] * Angle;
+			}
 		}
+		// Normalize Normal
+		Normal = Normal.GetSafeNormal();
 	}
+	else
+	{
+		//We use existing normals so just use all polygons having a vertex instance at the same location sharing the same normals
+		Normal = NormalRef;
+		TArray<FVertexInstanceID> VertexInstanceIDs = GetVertexVertexInstances(VertexID);
+		for (const FVertexInstanceID& ConnectedVertexInstanceID : VertexInstanceIDs)
+		{
+			if (ConnectedVertexInstanceID != VertexInstanceID && !VertexNormals[ConnectedVertexInstanceID].Equals(Normal))
+			{
+				continue;
+			}
 
-	// Normalize Normal
-	Normal = Normal.GetSafeNormal();
+			const TArray<FPolygonID>& ConnectedPolygons = GetVertexInstanceConnectedPolygons(ConnectedVertexInstanceID);
+			for (const FPolygonID ConnectedPolygonID : ConnectedPolygons)
+			{
+				const float Angle = bUseWeightedNormals ? GetPolygonCornerAngleForVertex(ConnectedPolygonID, VertexID) : 1.0f;
+
+				// If this polygon is actually connected to the vertex instance we're processing, also include its contributions towards the tangent
+				Tangent += PolygonTangents[ConnectedPolygonID] * Angle;
+				Binormal += PolygonBinormals[ConnectedPolygonID] * Angle;
+			}
+		}
+	}
+	
+	
 	float BinormalSign = 1.0f;
 	if (bComputeTangents)
 	{
@@ -949,7 +979,7 @@ void UMeshDescription::ComputeTangentsAndNormals(const FVertexInstanceID& Vertex
 	}
 }
 
-void UMeshDescription::ComputeTangentsAndNormals(const TArray<FVertexInstanceID>& VertexInstanceIDs, bool bComputeTangents, bool bUseWeightedNormals)
+void UMeshDescription::ComputeTangentsAndNormals(const TArray<FVertexInstanceID>& VertexInstanceIDs, EComputeNTBsOptions ComputeNTBsOptions)
 {
 	const TPolygonAttributeArray<FVector>& PolygonNormals = PolygonAttributes().GetAttributes<FVector>(MeshAttribute::Polygon::Normal);
 	const TPolygonAttributeArray<FVector>& PolygonTangents = PolygonAttributes().GetAttributes<FVector>(MeshAttribute::Polygon::Tangent);
@@ -961,11 +991,11 @@ void UMeshDescription::ComputeTangentsAndNormals(const TArray<FVertexInstanceID>
 
 	for (const FVertexInstanceID VertexInstanceID : VertexInstanceIDs)
 	{
-		ComputeTangentsAndNormals(VertexInstanceID, bComputeTangents, bUseWeightedNormals, PolygonNormals, PolygonTangents, PolygonBinormals, VertexNormals, VertexTangents, VertexBinormalSigns);
+		ComputeTangentsAndNormals(VertexInstanceID, ComputeNTBsOptions, PolygonNormals, PolygonTangents, PolygonBinormals, VertexNormals, VertexTangents, VertexBinormalSigns);
 	}
 }
 
-void UMeshDescription::ComputeTangentsAndNormals(bool bComputeTangents, bool bUseWeightedNormals)
+void UMeshDescription::ComputeTangentsAndNormals(EComputeNTBsOptions ComputeNTBsOptions)
 {
 	const TPolygonAttributeArray<FVector>& PolygonNormals = PolygonAttributes().GetAttributes<FVector>(MeshAttribute::Polygon::Normal);
 	const TPolygonAttributeArray<FVector>& PolygonTangents = PolygonAttributes().GetAttributes<FVector>(MeshAttribute::Polygon::Tangent);
@@ -977,7 +1007,7 @@ void UMeshDescription::ComputeTangentsAndNormals(bool bComputeTangents, bool bUs
 
 	for (const FVertexInstanceID VertexInstanceID : VertexInstances().GetElementIDs())
 	{
-		ComputeTangentsAndNormals(VertexInstanceID, bComputeTangents, bUseWeightedNormals, PolygonNormals, PolygonTangents, PolygonBinormals, VertexNormals, VertexTangents, VertexBinormalSigns);
+		ComputeTangentsAndNormals(VertexInstanceID, ComputeNTBsOptions, PolygonNormals, PolygonTangents, PolygonBinormals, VertexNormals, VertexTangents, VertexBinormalSigns);
 	}
 }
 
diff --git a/Engine/Source/Runtime/MeshDescription/Public/MeshDescription.h b/Engine/Source/Runtime/MeshDescription/Public/MeshDescription.h
--- a/Engine/Source/Runtime/MeshDescription/Public/MeshDescription.h
+++ b/Engine/Source/Runtime/MeshDescription/Public/MeshDescription.h
@@ -230,6 +230,15 @@ using FEdgeArray = TMeshElementArray<FMeshEdge, FEdgeID>;
 using FPolygonArray = TMeshElementArray<FMeshPolygon, FPolygonID>;
 using FPolygonGroupArray = TMeshElementArray<FMeshPolygonGroup, FPolygonGroupID>;
 
+UENUM()
+enum class EComputeNTBsOptions : uint32
+{
+	None = 0x00000000,	// No flags
+	Normals = 0x00000001, //Compute the normals
+	Tangents = 0x00000002, //Compute the tangents
+	WeightedNTBs = 0x00000004, //Use weight angle when computing NTBs to proportionally distribute the vertex instance contribution to the normal/tangent/binormal in a smooth group.    i.e. Weight solve the cylinder problem
+};
+ENUM_CLASS_FLAGS(EComputeNTBsOptions);
 
 #define MESHDESCRIPTION_VER TEXT("264A91765090413DA61FE2A819147655")
 
@@ -805,9 +814,9 @@ public:
 	void ComputePolygonTangentsAndNormals(float ComparisonThreshold = 0.0f);
 	
 	/** Set the vertex instance tangent and normal only for the specified VertexInstanceIDs */
-	void ComputeTangentsAndNormals(const TArray<FVertexInstanceID>& VertexInstanceIDs, bool bComputeTangents, bool bUseWeightedNormals);
+	void ComputeTangentsAndNormals(const TArray<FVertexInstanceID>& VertexInstanceIDs, EComputeNTBsOptions ComputeNTBsOptions);
 	/** Set the vertex instance tangent and normal for all vertex instances in the mesh description. */
-	void ComputeTangentsAndNormals(bool bComputeTangents, bool bUseWeightedNormals);
+	void ComputeTangentsAndNormals(EComputeNTBsOptions ComputeNTBsOptions);
 
 	void ReversePolygonFacing(const FPolygonID PolygonID);
 	void ReverseAllPolygonFacing();
@@ -835,8 +844,7 @@ private:
 	void GetConnectedSoftEdges(const FVertexID VertexID, TArray<FEdgeID>& OutConnectedSoftEdges) const;
 	void ComputeTangentsAndNormals(
 		  const FVertexInstanceID& VertexInstanceID
-		, bool bComputeTangents
-		, bool bUseWeightedNormals
+		, EComputeNTBsOptions ComputeNTBsOptions
 		, const TPolygonAttributeArray<FVector>& PolygonNormals
 		, const TPolygonAttributeArray<FVector>& PolygonTangents
 		, const TPolygonAttributeArray<FVector>& PolygonBinormals
