Fix cook build

diff --git a/Engine/Source/Developer/MeshBuilder/Private/StaticMeshBuilder.cpp b/Engine/Source/Developer/MeshBuilder/Private/StaticMeshBuilder.cpp
--- a/Engine/Source/Developer/MeshBuilder/Private/StaticMeshBuilder.cpp
+++ b/Engine/Source/Developer/MeshBuilder/Private/StaticMeshBuilder.cpp
@@ -68,7 +68,6 @@ bool FStaticMeshBuilder::Build(UStaticMesh* StaticMesh, const FStaticMeshLODGrou
 
 		if (OriginalMeshDescription != nullptr)
 		{
-			OriginalMeshDescription->TriangulateMesh(); //Make sure the original mesh description is triangulate.
 			MeshDescription = MeshDescriptionHelper.GetRenderMeshDescription(StaticMesh);
 		}
 		
diff --git a/Engine/Source/Editor/UnrealEd/Private/Fbx/FbxStaticMeshImport.cpp b/Engine/Source/Editor/UnrealEd/Private/Fbx/FbxStaticMeshImport.cpp
--- a/Engine/Source/Editor/UnrealEd/Private/Fbx/FbxStaticMeshImport.cpp
+++ b/Engine/Source/Editor/UnrealEd/Private/Fbx/FbxStaticMeshImport.cpp
@@ -1741,6 +1741,13 @@ void UnFbx::FFbxImporter::PostImportStaticMesh(UStaticMesh* StaticMesh, TArray<F
 	IConsoleVariable* CVarDistanceFieldInterface = IConsoleManager::Get().FindConsoleVariable(TEXT("r.GenerateMeshDistanceFields"));
 	bool bOriginalGenerateMeshDistanceField = StaticMesh->bGenerateMeshDistanceField;
 	
+	//Always triangulate the original mesh description after we import it
+	UMeshDescription* MeshDescription = StaticMesh->GetOriginalMeshDescription(LODIndex);
+	if (MeshDescription)
+	{
+		MeshDescription->TriangulateMesh();
+	}
+
 	//Prebuild the static mesh when we use LodGroup and we want to modify the LodNumber
 	if (!ImportOptions->bImportScene)
 	{
diff --git a/Engine/Source/Runtime/Engine/Private/StaticMesh.cpp b/Engine/Source/Runtime/Engine/Private/StaticMesh.cpp
--- a/Engine/Source/Runtime/Engine/Private/StaticMesh.cpp
+++ b/Engine/Source/Runtime/Engine/Private/StaticMesh.cpp
@@ -1313,7 +1313,7 @@ FArchive& operator<<(FArchive& Ar, FMeshBuildSettings& BuildSettings)
 // differences, etc.) replace the version GUID below with a new one.
 // In case of merge conflicts with DDC versions, you MUST generate a new GUID
 // and set this new GUID as the version.                                       
-#define STATICMESH_DERIVEDDATA_VER TEXT("8CAD439CF4B14211ACC831AE9C7C98BC")
+#define STATICMESH_DERIVEDDATA_VER TEXT("D6FDB3200D904958970912FE16AAA185")
 
 static const FString& GetStaticMeshDerivedDataVersion()
 {
diff --git a/Engine/Source/Runtime/MeshDescription/Private/MeshDescription.cpp b/Engine/Source/Runtime/MeshDescription/Private/MeshDescription.cpp
--- a/Engine/Source/Runtime/MeshDescription/Private/MeshDescription.cpp
+++ b/Engine/Source/Runtime/MeshDescription/Private/MeshDescription.cpp
@@ -28,6 +28,14 @@ void UMeshDescription::Serialize( FArchive& Ar )
 	Ar << PolygonGroupAttributesSet;
 }
 
+void UMeshDescription::PostEditChangeProperty(FPropertyChangedEvent& PropertyChangedEvent)
+{
+	Super::PostEditChangeProperty(PropertyChangedEvent);
+	
+	//We need to triangulate the mesh every time the mesh description change
+	TriangulateMesh();
+}
+
 void UMeshDescription::Empty()
 {
 	VertexArray.Reset();
@@ -72,6 +80,15 @@ FString UMeshDescription::GetIdString()
 	TPolygonGroupAttributeArray<bool>& PolygonGroupEnableCollisions = PolygonGroupAttributes().GetAttributes<bool>(MeshAttribute::PolygonGroup::EnableCollision);
 	TPolygonGroupAttributeArray<bool>& PolygonGroupCastShadows = PolygonGroupAttributes().GetAttributes<bool>(MeshAttribute::PolygonGroup::CastShadow);
 	
+	
+	//Fail safe, we always want to compute the string ID with the triangle triangulate
+	if(GetPolygon(Polygons().GetFirstValidID()).Triangles.Num() == 0)
+	{
+		//We should not get here since Triangulate is call when we import staticMesh and in the post edit change of the MeshDescription
+		ensure(GetPolygon(Polygons().GetFirstValidID()).Triangles.Num() != 0);
+		TriangulateMesh();
+	}
+
 	auto AddItemToSha = [&](auto ItemSerialization)
 	{
 		FMemoryWriter Ar(TempBytes);
diff --git a/Engine/Source/Runtime/MeshDescription/Public/MeshDescription.h b/Engine/Source/Runtime/MeshDescription/Public/MeshDescription.h
--- a/Engine/Source/Runtime/MeshDescription/Public/MeshDescription.h
+++ b/Engine/Source/Runtime/MeshDescription/Public/MeshDescription.h
@@ -256,6 +256,7 @@ public:
 
 	// UObject interface
 	virtual void Serialize( FArchive& Ar ) override;
+	virtual void PostEditChangeProperty(struct FPropertyChangedEvent& PropertyChangedEvent) override;
 
 	//Empty the meshdescription
 	void Empty();
