Fix non editor build

diff --git a/Engine/Source/Runtime/MeshDescription/Private/MeshDescription.cpp b/Engine/Source/Runtime/MeshDescription/Private/MeshDescription.cpp
--- a/Engine/Source/Runtime/MeshDescription/Private/MeshDescription.cpp
+++ b/Engine/Source/Runtime/MeshDescription/Private/MeshDescription.cpp
@@ -28,6 +28,7 @@ void UMeshDescription::Serialize( FArchive& Ar )
 	Ar << PolygonGroupAttributesSet;
 }
 
+#if WITH_EDITOR
 void UMeshDescription::PostEditChangeProperty(FPropertyChangedEvent& PropertyChangedEvent)
 {
 	Super::PostEditChangeProperty(PropertyChangedEvent);
@@ -35,6 +36,7 @@ void UMeshDescription::PostEditChangeProperty(FPropertyChangedEvent& PropertyCha
 	//We need to triangulate the mesh every time the mesh description change
 	TriangulateMesh();
 }
+#endif // WITH_EDITOR
 
 void UMeshDescription::Empty()
 {
diff --git a/Engine/Source/Runtime/MeshDescription/Public/MeshDescription.h b/Engine/Source/Runtime/MeshDescription/Public/MeshDescription.h
--- a/Engine/Source/Runtime/MeshDescription/Public/MeshDescription.h
+++ b/Engine/Source/Runtime/MeshDescription/Public/MeshDescription.h
@@ -256,7 +256,9 @@ public:
 
 	// UObject interface
 	virtual void Serialize( FArchive& Ar ) override;
+#if WITH_EDITOR
 	virtual void PostEditChangeProperty(struct FPropertyChangedEvent& PropertyChangedEvent) override;
+#endif
 
 	//Empty the meshdescription
 	void Empty();
