- Added UMeshDescription::DetermineEdgeHardnessesFromVertexInstanceNormals() - discerns edge hardnesses from the normals currently set on vertex instances. Use this to automatically set edge hardnesses on a mesh when building a mesh description.
- Added a new edge attribute: IsUVSeam.
- Added UMeshDescription::DetermineUVSeamsFromUVs() - to automatically determine which edges are UV seams in the mesh.
- Added edge coloring in the wireframe overlay to show UV seams as well as edge hardness.

diff --git a/Engine/Plugins/Editor/MeshEditor/Source/MeshEditor/MeshEditorStaticMeshAdapter.cpp b/Engine/Plugins/Editor/MeshEditor/Source/MeshEditor/MeshEditorStaticMeshAdapter.cpp
--- a/Engine/Plugins/Editor/MeshEditor/Source/MeshEditor/MeshEditorStaticMeshAdapter.cpp
+++ b/Engine/Plugins/Editor/MeshEditor/Source/MeshEditor/MeshEditorStaticMeshAdapter.cpp
@@ -6,9 +6,16 @@
 #include "WireframeMeshComponent.h"
 
 
-static FColor GetEdgeColor( bool bIsHardEdge )
+static FColor GetEdgeColor( bool bIsHardEdge, bool bIsUVSeam )
 {
-	return bIsHardEdge ? FColor( 0, 76, 76 ) : FColor( 0, 0, 76 );
+	if( bIsUVSeam )
+	{
+		return bIsHardEdge ? FColor( 20, 80, 25 ) : FColor( 0, 46, 10 );
+	}
+	else
+	{
+		return bIsHardEdge ? FColor( 20, 20, 128 ) : FColor( 0, 0, 76 );
+	}
 }
 
 
@@ -61,11 +68,12 @@ void UMeshEditorStaticMeshAdapter::OnRebuildRenderMesh( const UEditableMesh* Edi
 
 	// Add all edges
 	const TEdgeAttributeArray<bool>& EdgeHardnesses = MeshDescription->EdgeAttributes().GetAttributes<bool>( MeshAttribute::Edge::IsHard );
+	const TEdgeAttributeArray<bool>& EdgeUVSeams = MeshDescription->EdgeAttributes().GetAttributes<bool>( MeshAttribute::Edge::IsUVSeam );
 	for( const FEdgeID EdgeID : MeshDescription->Edges().GetElementIDs() )
 	{
 		WireframeMesh->AddEdge( EdgeID );
 		WireframeMesh->SetEdgeVertices( EdgeID, EditableMesh->GetEdgeVertex( EdgeID, 0 ), EditableMesh->GetEdgeVertex( EdgeID, 1 ) );
-		WireframeMesh->SetEdgeColor( EdgeID, GetEdgeColor( EdgeHardnesses[ EdgeID ] ) );
+		WireframeMesh->SetEdgeColor( EdgeID, GetEdgeColor( EdgeHardnesses[ EdgeID ], EdgeUVSeams[ EdgeID ] ) );
 	}
 
 	// Add all polygons and edge instances
@@ -223,12 +231,13 @@ void UMeshEditorStaticMeshAdapter::OnCreateEdges( const UEditableMesh* EditableM
 	check( MeshDescription );
 
 	const TEdgeAttributeArray<bool>& EdgeHardnesses = MeshDescription->EdgeAttributes().GetAttributes<bool>( MeshAttribute::Edge::IsHard );
+	const TEdgeAttributeArray<bool>& EdgeUVSeams = MeshDescription->EdgeAttributes().GetAttributes<bool>( MeshAttribute::Edge::IsUVSeam );
 
 	for( const FEdgeID EdgeID : EdgeIDs )
 	{
 		WireframeMesh->AddEdge( EdgeID );
 		WireframeMesh->SetEdgeVertices( EdgeID, EditableMesh->GetEdgeVertex( EdgeID, 0 ), EditableMesh->GetEdgeVertex( EdgeID, 1 ) );
-		WireframeMesh->SetEdgeColor( EdgeID, GetEdgeColor( EdgeHardnesses[ EdgeID ] ) );
+		WireframeMesh->SetEdgeColor( EdgeID, GetEdgeColor( EdgeHardnesses[ EdgeID ], EdgeUVSeams[ EdgeID ] ) );
 
 		for( const FPolygonID PolygonID : MeshDescription->GetEdgeConnectedPolygons( EdgeID ) )
 		{
@@ -266,9 +275,16 @@ void UMeshEditorStaticMeshAdapter::OnSetEdgesVertices( const UEditableMesh* Edit
 
 void UMeshEditorStaticMeshAdapter::OnSetEdgeAttribute( const UEditableMesh* EditableMesh, const FEdgeID EdgeID, const FMeshElementAttributeData& Attribute )
 {
-	if( Attribute.AttributeName == MeshAttribute::Edge::IsHard )
+	const UMeshDescription* MeshDescription = EditableMesh->GetMeshDescription();
+	check( MeshDescription );
+
+	const TEdgeAttributeArray<bool>& EdgeHardnesses = MeshDescription->EdgeAttributes().GetAttributes<bool>( MeshAttribute::Edge::IsHard );
+	const TEdgeAttributeArray<bool>& EdgeUVSeams = MeshDescription->EdgeAttributes().GetAttributes<bool>( MeshAttribute::Edge::IsUVSeam );
+
+	if( Attribute.AttributeName == MeshAttribute::Edge::IsHard ||
+		Attribute.AttributeName == MeshAttribute::Edge::IsUVSeam )
 	{
-		WireframeMesh->SetEdgeColor( EdgeID, GetEdgeColor( Attribute.AttributeValue.GetValue<bool>() ) );
+		WireframeMesh->SetEdgeColor( EdgeID, GetEdgeColor( EdgeHardnesses[ EdgeID ], EdgeUVSeams[ EdgeID ] ) );
 	}
 }
 
diff --git a/Engine/Plugins/Runtime/EditableMesh/Source/EditableMesh/EditableStaticMeshAdapter.cpp b/Engine/Plugins/Runtime/EditableMesh/Source/EditableMesh/EditableStaticMeshAdapter.cpp
--- a/Engine/Plugins/Runtime/EditableMesh/Source/EditableMesh/EditableStaticMeshAdapter.cpp
+++ b/Engine/Plugins/Runtime/EditableMesh/Source/EditableMesh/EditableStaticMeshAdapter.cpp
@@ -260,9 +260,6 @@ void UEditableStaticMeshAdapter::InitEditableStaticMesh( UEditableMesh* Editable
 
 				const FIndexArrayView RenderingIndices = StaticMeshLOD.IndexBuffer.GetArrayView();
 
-				static TMultiMap<FEdgeID, TTuple<FVertexInstanceID, FVertexInstanceID>> EdgeToVertexInstancePair;
-				EdgeToVertexInstancePair.Reset();
-
 				const uint32 NumSections = StaticMeshLOD.Sections.Num();
 				MeshDescription->ReserveNewPolygonGroups( NumSections );
 
@@ -334,38 +331,6 @@ void UEditableStaticMeshAdapter::InitEditableStaticMesh( UEditableMesh* Editable
 									}
 
 									NewEdgeIDs[ TriangleEdgeNumber ] = NewEdgeID;
-
-									// Determine edge hardness by checking whether overlapping edges' vertex instances have the same normals or not
-									if( !EdgeHardnesses[ NewEdgeID ] )
-									{
-										// Get vertex instance IDs corresponding to the ends of the edge.
-										// Always order them the same for overlapping edges so we compare equivalent vertex instances each time.
-										// For simplicity, we just order them by vertex ID value, lowest/highest.
-										FVertexInstanceID VertexInstanceID0 = TriangleVertexInstanceIDs[ TriangleEdgeNumber ];
-										FVertexInstanceID VertexInstanceID1 = TriangleVertexInstanceIDs[ ( TriangleEdgeNumber + 1 ) % 3 ];
-										if( VertexID0.GetValue() > VertexID1.GetValue() )
-										{
-											Swap( VertexInstanceID0, VertexInstanceID1 );
-										}
-
-										TArray<TTuple<FVertexInstanceID, FVertexInstanceID>> VertexInstancePairs;
-										EdgeToVertexInstancePair.MultiFind( NewEdgeID, VertexInstancePairs );
-
-										for( const TTuple<FVertexInstanceID, FVertexInstanceID>& VertexInstancePair : VertexInstancePairs )
-										{
-											const FVertexInstanceID ExistingVertexInstanceID0 = VertexInstancePair.Get<0>();
-											const FVertexInstanceID ExistingVertexInstanceID1 = VertexInstancePair.Get<1>();
-
-											if( VertexInstanceNormals[ ExistingVertexInstanceID0 ] != VertexInstanceNormals[ VertexInstanceID0 ] ||
-												VertexInstanceNormals[ ExistingVertexInstanceID1 ] != VertexInstanceNormals[ VertexInstanceID1 ] )
-											{
-												EdgeHardnesses[ NewEdgeID ] = true;
-												break;
-											}
-										}
-
-										EdgeToVertexInstancePair.Add( NewEdgeID, MakeTuple( VertexInstanceID0, VertexInstanceID1 ) );
-									}
 								}
 							}
 
@@ -408,36 +373,45 @@ void UEditableStaticMeshAdapter::InitEditableStaticMesh( UEditableMesh* Editable
 						}
 					}
 				}
-			}
-		}
-	}
 
-	this->CachedBoundingBoxAndSphere = StaticMesh->GetBounds();
+				// Determine edge hardnesses
+				MeshDescription->DetermineEdgeHardnessesFromVertexInstanceNormals();
 
-	// Cache polygon tangent bases
-	static TArray<FPolygonID> PolygonIDs;
-	PolygonIDs.Reset();
-	for( const FPolygonID PolygonID : EditableMesh->GetMeshDescription()->Polygons().GetElementIDs() )
-	{
-		PolygonIDs.Add( PolygonID );
-	}
+				// Determine UV seams
+				if( NumUVs > 0 )
+				{
+					MeshDescription->DetermineUVSeamsFromUVs( 0 );
+				}
+
+				// Cache polygon tangent bases
+				static TArray<FPolygonID> PolygonIDs;
+				PolygonIDs.Reset();
+				for( const FPolygonID PolygonID : EditableMesh->GetMeshDescription()->Polygons().GetElementIDs() )
+				{
+					PolygonIDs.Add( PolygonID );
+				}
 
-	EditableMesh->GeneratePolygonTangentsAndNormals( PolygonIDs );
+				EditableMesh->GeneratePolygonTangentsAndNormals( PolygonIDs );
 
 #if 0
-	// Test tangent generation
-	static TArray<FPolygonID> AllPolygons;
-	AllPolygons.Reset();
-	for( int32 PolygonIndex = 0; PolygonIndex < EditableMesh->Polygons.GetMaxIndex(); ++PolygonIndex )
-	{
-		if( EditableMesh->Polygons.IsAllocated( PolygonIndex ) )
-		{
-			EditableMesh->PolygonsPendingNewTangentBasis.Add( FPolygonID( PolygonIndex ) );
+				// Test tangent generation
+				static TArray<FPolygonID> AllPolygons;
+				AllPolygons.Reset();
+				for( int32 PolygonIndex = 0; PolygonIndex < EditableMesh->Polygons.GetMaxIndex(); ++PolygonIndex )
+				{
+					if( EditableMesh->Polygons.IsAllocated( PolygonIndex ) )
+					{
+						EditableMesh->PolygonsPendingNewTangentBasis.Add( FPolygonID( PolygonIndex ) );
+					}
+				}
+
+				EditableMesh->GenerateTangentsAndNormals();
+#endif
+			}
 		}
 	}
 
-	EditableMesh->GenerateTangentsAndNormals();
-#endif
+	this->CachedBoundingBoxAndSphere = StaticMesh->GetBounds();
 
 #if EDITABLE_MESH_USE_OPENSUBDIV
 	EditableMesh->RefreshOpenSubdiv();
diff --git a/Engine/Source/Runtime/Engine/Private/StaticMesh.cpp b/Engine/Source/Runtime/Engine/Private/StaticMesh.cpp
--- a/Engine/Source/Runtime/Engine/Private/StaticMesh.cpp
+++ b/Engine/Source/Runtime/Engine/Private/StaticMesh.cpp
@@ -2507,6 +2507,7 @@ void UStaticMesh::RegisterMeshAttributes( UMeshDescription* MeshDescription )
 
 	// Add basic edge attributes
 	MeshDescription->EdgeAttributes().RegisterAttribute<bool>( MeshAttribute::Edge::IsHard, 1, false );
+	MeshDescription->EdgeAttributes().RegisterAttribute<bool>( MeshAttribute::Edge::IsUVSeam, 1, false );
 	MeshDescription->EdgeAttributes().RegisterAttribute<float>( MeshAttribute::Edge::CreaseSharpness, 1, 0.0f, EMeshAttributeFlags::Lerpable );
 
 	// Add basic polygon attributes
diff --git a/Engine/Source/Runtime/MeshDescription/Private/MeshAttributes.cpp b/Engine/Source/Runtime/MeshDescription/Private/MeshAttributes.cpp
--- a/Engine/Source/Runtime/MeshDescription/Private/MeshAttributes.cpp
+++ b/Engine/Source/Runtime/MeshDescription/Private/MeshAttributes.cpp
@@ -14,6 +14,7 @@ namespace MeshAttribute
 	const FName VertexInstance::Color( "Color" );
 
 	const FName Edge::IsHard( "IsHard" );
+	const FName Edge::IsUVSeam( "IsUVSeam" );
 	const FName Edge::CreaseSharpness( "CreaseSharpness" );
 
 	const FName Polygon::Normal( "Normal" );
diff --git a/Engine/Source/Runtime/MeshDescription/Private/MeshDescription.cpp b/Engine/Source/Runtime/MeshDescription/Private/MeshDescription.cpp
--- a/Engine/Source/Runtime/MeshDescription/Private/MeshDescription.cpp
+++ b/Engine/Source/Runtime/MeshDescription/Private/MeshDescription.cpp
@@ -1030,6 +1030,132 @@ void UMeshDescription::ComputeTangentsAndNormals(EComputeNTBsOptions ComputeNTBs
 	}
 }
 
+
+void UMeshDescription::DetermineEdgeHardnessesFromVertexInstanceNormals( const float Tolerance )
+{
+	const TVertexInstanceAttributeArray<FVector>& VertexNormals = VertexInstanceAttributes().GetAttributes<FVector>( MeshAttribute::VertexInstance::Normal );
+	TEdgeAttributeArray<bool>& EdgeHardnesses = EdgeAttributes().GetAttributes<bool>( MeshAttribute::Edge::IsHard );
+
+	// Holds unique vertex instance IDs for a given edge vertex
+	TArray<FVertexInstanceID> UniqueVertexInstanceIDs;
+
+	for( const FEdgeID EdgeID : Edges().GetElementIDs() )
+	{
+		// Get list of polygons connected to this edge
+		const TArray<FPolygonID>& ConnectedPolygonIDs = GetEdgeConnectedPolygons( EdgeID );
+		if( ConnectedPolygonIDs.Num() == 0 )
+		{
+			// What does it mean if an edge has no connected polygons? For now we just skip it
+			continue;
+		}
+
+		// Assume by default that the edge is soft - but as soon as any vertex instance belonging to a connected polygon
+		// has a distinct normal from the others (within the given tolerance), we mark it as hard.
+		// The exception is if an edge has exactly one connected polygon: in this case we automatically deem it a hard edge. 
+		bool bEdgeIsHard = ( ConnectedPolygonIDs.Num() == 1 );
+
+		// Examine vertices on each end of the edge, if we haven't yet identified it as 'hard'
+		for( int32 VertexIndex = 0; !bEdgeIsHard && VertexIndex < 2; ++VertexIndex )
+		{
+			const FVertexID VertexID = GetEdgeVertex( EdgeID, VertexIndex );
+
+			const int32 ReservedElements = 4;
+			UniqueVertexInstanceIDs.Reset( ReservedElements );
+
+			// Get a list of all vertex instances for this vertex which form part of any polygon connected to the edge
+			for( const FVertexInstanceID VertexInstanceID : GetVertexVertexInstances( VertexID ) )
+			{
+				for( const FPolygonID PolygonID : GetVertexInstanceConnectedPolygons( VertexInstanceID ) )
+				{
+					if( ConnectedPolygonIDs.Contains( PolygonID ) )
+					{
+						UniqueVertexInstanceIDs.AddUnique( VertexInstanceID );
+						break;
+					}
+				}
+			}
+			check( UniqueVertexInstanceIDs.Num() > 0 );
+
+			// First unique vertex instance is used as a reference against which the others are compared.
+			// (not a perfect approach: really the 'median' should be used as a reference)
+			const FVector ReferenceNormal = VertexNormals[ UniqueVertexInstanceIDs[ 0 ] ];
+			for( int32 Index = 1; Index < UniqueVertexInstanceIDs.Num(); ++Index )
+			{
+				if( !VertexNormals[ UniqueVertexInstanceIDs[ Index ] ].Equals( ReferenceNormal, Tolerance ) )
+				{
+					bEdgeIsHard = true;
+					break;
+				}
+			}
+		}
+
+		EdgeHardnesses[ EdgeID ] = bEdgeIsHard;
+	}
+}
+
+
+void UMeshDescription::DetermineUVSeamsFromUVs( const int32 UVIndex, const float Tolerance )
+{
+	const TVertexInstanceAttributeArray<FVector2D>& VertexUVs = VertexInstanceAttributes().GetAttributes<FVector2D>( MeshAttribute::VertexInstance::TextureCoordinate, UVIndex );
+	TEdgeAttributeArray<bool>& EdgeUVSeams = EdgeAttributes().GetAttributes<bool>( MeshAttribute::Edge::IsUVSeam );
+
+	// Holds unique vertex instance IDs for a given edge vertex
+	TArray<FVertexInstanceID> UniqueVertexInstanceIDs;
+
+	for( const FEdgeID EdgeID : Edges().GetElementIDs() )
+	{
+		// Get list of polygons connected to this edge
+		const TArray<FPolygonID>& ConnectedPolygonIDs = GetEdgeConnectedPolygons( EdgeID );
+		if( ConnectedPolygonIDs.Num() == 0 )
+		{
+			// What does it mean if an edge has no connected polygons? For now we just skip it
+			continue;
+		}
+
+		// Assume by default that the edge is not a UV seam - but as soon as any vertex instance belonging to a connected polygon
+		// has a distinct UV from the others (within the given tolerance), we mark it as a UV seam.
+		bool bEdgeIsUVSeam = false;
+
+		// Examine vertices on each end of the edge, if we haven't yet identified it as a UV seam
+		for( int32 VertexIndex = 0; !bEdgeIsUVSeam && VertexIndex < 2; ++VertexIndex )
+		{
+			const FVertexID VertexID = GetEdgeVertex( EdgeID, VertexIndex );
+
+			const int32 ReservedElements = 4;
+			UniqueVertexInstanceIDs.Reset( ReservedElements );
+
+			// Get a list of all vertex instances for this vertex which form part of any polygon connected to the edge
+			for( const FVertexInstanceID VertexInstanceID : GetVertexVertexInstances( VertexID ) )
+			{
+				for( const FPolygonID PolygonID : GetVertexInstanceConnectedPolygons( VertexInstanceID ) )
+				{
+					if( ConnectedPolygonIDs.Contains( PolygonID ) )
+					{
+						UniqueVertexInstanceIDs.AddUnique( VertexInstanceID );
+						break;
+					}
+				}
+			}
+			check( UniqueVertexInstanceIDs.Num() > 0 );
+
+			// First unique vertex instance is used as a reference against which the others are compared.
+			// (not a perfect approach: really the 'median' should be used as a reference)
+			const FVector2D ReferenceUV = VertexUVs[ UniqueVertexInstanceIDs[ 0 ] ];
+			for( int32 Index = 1; Index < UniqueVertexInstanceIDs.Num(); ++Index )
+			{
+				if( !VertexUVs[ UniqueVertexInstanceIDs[ Index ] ].Equals( ReferenceUV, Tolerance ) )
+				{
+					bEdgeIsUVSeam = true;
+					break;
+				}
+			}
+		}
+
+		EdgeUVSeams[ EdgeID ] = bEdgeIsUVSeam;
+	}
+}
+
+
 void UMeshDescription::ReversePolygonFacing(const FPolygonID PolygonID)
 {
 	FMeshPolygon& Polygon = GetPolygon(PolygonID);
diff --git a/Engine/Source/Runtime/MeshDescription/Public/MeshAttributeArray.h b/Engine/Source/Runtime/MeshDescription/Public/MeshAttributeArray.h
--- a/Engine/Source/Runtime/MeshDescription/Public/MeshAttributeArray.h
+++ b/Engine/Source/Runtime/MeshDescription/Public/MeshAttributeArray.h
@@ -42,9 +42,6 @@ public:
 
 	/** Disallow use of the copy constructor to prevent arrays being mistakenly accessed in the UMeshDescription by value */
 	TMeshAttributeArrayBase( const TMeshAttributeArrayBase& ) = delete;
-//	{
-//		static_assert( false, "TMeshAttributeArray must be accessed by reference");
-//	}
 
 	/**
 	 * Custom serialization for TMeshAttributeArrayBase.
diff --git a/Engine/Source/Runtime/MeshDescription/Public/MeshAttributes.h b/Engine/Source/Runtime/MeshDescription/Public/MeshAttributes.h
--- a/Engine/Source/Runtime/MeshDescription/Public/MeshAttributes.h
+++ b/Engine/Source/Runtime/MeshDescription/Public/MeshAttributes.h
@@ -27,6 +27,7 @@ namespace MeshAttribute
 	namespace Edge
 	{
 		extern MESHDESCRIPTION_API const FName IsHard;
+		extern MESHDESCRIPTION_API const FName IsUVSeam;
 		extern MESHDESCRIPTION_API const FName CreaseSharpness;
 	}
 
diff --git a/Engine/Source/Runtime/MeshDescription/Public/MeshDescription.h b/Engine/Source/Runtime/MeshDescription/Public/MeshDescription.h
--- a/Engine/Source/Runtime/MeshDescription/Public/MeshDescription.h
+++ b/Engine/Source/Runtime/MeshDescription/Public/MeshDescription.h
@@ -821,6 +821,12 @@ public:
 	/** Set the vertex instance tangent and normal for all vertex instances in the mesh description. */
 	void ComputeTangentsAndNormals(EComputeNTBsOptions ComputeNTBsOptions);
 
+	/** Determine the edge hardnesses from existing normals */
+	void DetermineEdgeHardnessesFromVertexInstanceNormals( const float Tolerance = KINDA_SMALL_NUMBER );
+
+	/** Determine UV seams from existing vertex instance UVs */
+	void DetermineUVSeamsFromUVs( const int32 UVIndex, const float Tolerance = KINDA_SMALL_NUMBER );
+
 	void ReversePolygonFacing(const FPolygonID PolygonID);
 	void ReverseAllPolygonFacing();
 
