Added methods in UMeshDescription for getting the polygon IDs which comprise a particular UV chart, and for getting all UV charts.

diff --git a/Engine/Source/Runtime/MeshDescription/Private/MeshDescription.cpp b/Engine/Source/Runtime/MeshDescription/Private/MeshDescription.cpp
--- a/Engine/Source/Runtime/MeshDescription/Private/MeshDescription.cpp
+++ b/Engine/Source/Runtime/MeshDescription/Private/MeshDescription.cpp
@@ -1037,6 +1037,7 @@ void UMeshDescription::DetermineEdgeHardnessesFromVertexInstanceNormals( const f
 	TEdgeAttributeArray<bool>& EdgeHardnesses = EdgeAttributes().GetAttributes<bool>( MeshAttribute::Edge::IsHard );
 
 	// Holds unique vertex instance IDs for a given edge vertex
+	// @todo: use TMemStackAllocator or similar to avoid expensive allocations
 	TArray<FVertexInstanceID> UniqueVertexInstanceIDs;
 
 	for( const FEdgeID EdgeID : Edges().GetElementIDs() )
@@ -1100,6 +1101,7 @@ void UMeshDescription::DetermineUVSeamsFromUVs( const int32 UVIndex, const float
 	TEdgeAttributeArray<bool>& EdgeUVSeams = EdgeAttributes().GetAttributes<bool>( MeshAttribute::Edge::IsUVSeam );
 
 	// Holds unique vertex instance IDs for a given edge vertex
+	// @todo: use TMemStackAllocator or similar to avoid expensive allocations
 	TArray<FVertexInstanceID> UniqueVertexInstanceIDs;
 
 	for( const FEdgeID EdgeID : Edges().GetElementIDs() )
@@ -1156,6 +1158,85 @@ void UMeshDescription::DetermineUVSeamsFromUVs( const int32 UVIndex, const float
 }
 
 
+void UMeshDescription::GetPolygonsInSameChartAsPolygon( const FPolygonID PolygonID, TArray<FPolygonID>& OutPolygonIDs )
+{
+	const TEdgeAttributeArray<bool>& EdgeUVSeams = EdgeAttributes().GetAttributes<bool>( MeshAttribute::Edge::IsUVSeam );
+	const int32 NumPolygons = Polygons().Num();
+
+	// This holds the results - all polygon IDs which are in the same UV chart
+	OutPolygonIDs.Reset( NumPolygons );
+
+	// This holds all the polygons we need to check, and those we have already checked so we don't add duplicates
+	// @todo: use TMemStackAllocator or similar to avoid expensive allocations
+	TArray<FPolygonID> PolygonsToCheck;
+	PolygonsToCheck.Reset( NumPolygons );
+
+	// Add the initial polygon
+	PolygonsToCheck.Add( PolygonID );
+
+	int32 Index = 0;
+	while( Index < PolygonsToCheck.Num() )
+	{
+		// Process the next polygon to be checked. If it's in this list, we already know it's one of the results. Now we have to check the neighbors.
+		const FPolygonID PolygonToCheck = PolygonsToCheck[ Index ];
+		OutPolygonIDs.Add( PolygonToCheck );
+		Index++;
+
+		// Iterate through edges of the polygon
+		const TArray<FVertexInstanceID>& VertexInstanceIDs = GetPolygonPerimeterVertexInstances( PolygonToCheck );
+		FVertexID LastVertexID = GetVertexInstanceVertex( VertexInstanceIDs.Last() );
+		for( const FVertexInstanceID VertexInstanceID : VertexInstanceIDs )
+		{
+			const FVertexID VertexID = GetVertexInstanceVertex( VertexInstanceID );
+			const FEdgeID EdgeID = GetVertexPairEdge( VertexID, LastVertexID );
+			if( EdgeID != FEdgeID::Invalid && !EdgeUVSeams[ EdgeID ] )
+			{
+				// If it's a valid edge and not a UV seam, check its connected polygons
+				const TArray<FPolygonID>& ConnectedPolygonIDs = GetEdgeConnectedPolygons( EdgeID );
+				for( const FPolygonID ConnectedPolygonID : ConnectedPolygonIDs )
+				{
+					// Add polygons which aren't the one being checked, and haven't already been added to the list
+					if( ConnectedPolygonID != PolygonToCheck && !PolygonsToCheck.Contains( ConnectedPolygonID ) )
+					{
+						PolygonsToCheck.Add( ConnectedPolygonID );
+					}
+				}
+			}
+			LastVertexID = VertexID;
+		}
+	}
+}
+
+
+void UMeshDescription::GetAllCharts( TArray<TArray<FPolygonID>>& OutCharts )
+{
+	// @todo: OutCharts: array of array doesn't seem like a really efficient data structure. Also templatize on allocator?
+
+	const int32 NumPolygons = Polygons().Num();
+
+	// Maintain a record of the polygons which have already been entered into a chart
+	// @todo: use TMemStackAllocator or similar to avoid expensive allocations
+	TSet<FPolygonID> ConsumedPolygons;
+	ConsumedPolygons.Reserve( NumPolygons );
+
+	for( const FPolygonID PolygonID : Polygons().GetElementIDs() )
+	{
+		if( !ConsumedPolygons.Contains( PolygonID ) )
+		{
+			OutCharts.Emplace();
+			TArray<FPolygonID>& Chart = OutCharts.Last();
+			GetPolygonsInSameChartAsPolygon( PolygonID, Chart );
+
+			// Mark all polygons in the chart as 'consumed'. Note that the chart will also contain the initial polygon.
+			for( const FPolygonID ChartPolygon : Chart )
+			{
+				ConsumedPolygons.Add( ChartPolygon );
+			}
+		}
+	}
+}
+
+
 void UMeshDescription::ReversePolygonFacing(const FPolygonID PolygonID)
 {
 	FMeshPolygon& Polygon = GetPolygon(PolygonID);
diff --git a/Engine/Source/Runtime/MeshDescription/Public/MeshDescription.h b/Engine/Source/Runtime/MeshDescription/Public/MeshDescription.h
--- a/Engine/Source/Runtime/MeshDescription/Public/MeshDescription.h
+++ b/Engine/Source/Runtime/MeshDescription/Public/MeshDescription.h
@@ -827,6 +827,12 @@ public:
 	/** Determine UV seams from existing vertex instance UVs */
 	void DetermineUVSeamsFromUVs( const int32 UVIndex, const float Tolerance = KINDA_SMALL_NUMBER );
 
+	/** Get polygons in the same UV chart as the specified polygon */
+	void GetPolygonsInSameChartAsPolygon( const FPolygonID PolygonID, TArray<FPolygonID>& OutPolygonIDs );
+
+	/** Get array of all UV charts */
+	void GetAllCharts( TArray<TArray<FPolygonID>>& OutCharts );
+
 	void ReversePolygonFacing(const FPolygonID PolygonID);
 	void ReverseAllPolygonFacing();
 
