Copied from Dev-Editor CL 4096445

New serialization layout for UMeshDescription.
- Only the bare minimum is serialized: any internal values which can be inferred from others in the Mesh Description are omitted.
- Triangles are no longer serialized: a triangulation step is performed per polygon when serialized.
- Attribute arrays of simple types are now serialized with BulkSerialize for speed; only FName requires element-by-element serialization.

diff --git a/Engine/Source/Runtime/Engine/Private/StaticMesh.cpp b/Engine/Source/Runtime/Engine/Private/StaticMesh.cpp
--- a/Engine/Source/Runtime/Engine/Private/StaticMesh.cpp
+++ b/Engine/Source/Runtime/Engine/Private/StaticMesh.cpp
@@ -2714,7 +2714,7 @@ void UStaticMesh::FixupMaterialSlotName()
 // differences, etc.) replace the version GUID below with a new one.
 // In case of merge conflicts with DDC versions, you MUST generate a new GUID
 // and set this new GUID as the version.                                       
-#define MESHDATAKEY_STATICMESH_DERIVEDDATA_VER TEXT("22DF8DBA439A45E0A0745FBE8415CEA2")
+#define MESHDATAKEY_STATICMESH_DERIVEDDATA_VER TEXT("278E9A60614C418FB701B4383D23CD6B")
 
 static const FString& GetMeshDataKeyStaticMeshDerivedDataVersion()
 {
diff --git a/Engine/Source/Runtime/MeshDescription/Private/MeshDescription.cpp b/Engine/Source/Runtime/MeshDescription/Private/MeshDescription.cpp
--- a/Engine/Source/Runtime/MeshDescription/Private/MeshDescription.cpp
+++ b/Engine/Source/Runtime/MeshDescription/Private/MeshDescription.cpp
@@ -26,6 +26,57 @@ void UMeshDescription::Serialize( FArchive& Ar )
 	Ar << EdgeAttributesSet;
 	Ar << PolygonAttributesSet;
 	Ar << PolygonGroupAttributesSet;
+
+	if( Ar.IsLoading() )
+	{
+		// Populate vertex instance IDs for vertices
+		for( const FVertexInstanceID VertexInstanceID : VertexInstanceArray.GetElementIDs() )
+		{
+			const FVertexID VertexID = GetVertexInstanceVertex( VertexInstanceID );
+			VertexArray[ VertexID ].VertexInstanceIDs.Add( VertexInstanceID );
+		}
+
+		// Populate edge IDs for vertices
+		for( const FEdgeID EdgeID : EdgeArray.GetElementIDs() )
+		{
+			const FVertexID VertexID0 = GetEdgeVertex( EdgeID, 0 );
+			const FVertexID VertexID1 = GetEdgeVertex( EdgeID, 1 );
+			VertexArray[ VertexID0 ].ConnectedEdgeIDs.Add( EdgeID );
+			VertexArray[ VertexID1 ].ConnectedEdgeIDs.Add( EdgeID );
+		}
+
+		// Populate polygon IDs for vertex instances, edges and polygon groups
+		for( const FPolygonID PolygonID : PolygonArray.GetElementIDs() )
+		{
+			auto PopulatePolygonIDs = [ this, PolygonID ]( const TArray<FVertexInstanceID>& VertexInstanceIDs )
+			{
+				const int32 NumVertexInstances = VertexInstanceIDs.Num();
+				for( int32 Index = 0; Index < NumVertexInstances; ++Index )
+				{
+					const FVertexInstanceID VertexInstanceID0 = VertexInstanceIDs[ Index ];
+					const FVertexInstanceID VertexInstanceID1 = VertexInstanceIDs[ ( Index + 1 ) % NumVertexInstances ];
+					const FVertexID VertexID0 = GetVertexInstanceVertex( VertexInstanceID0 );
+					const FVertexID VertexID1 = GetVertexInstanceVertex( VertexInstanceID1 );
+					const FEdgeID EdgeID = GetVertexPairEdge( VertexID0, VertexID1 );
+					VertexInstanceArray[ VertexInstanceID0 ].ConnectedPolygons.Add( PolygonID );
+					EdgeArray[ EdgeID ].ConnectedPolygons.Add( PolygonID );
+				}
+			};
+
+			PopulatePolygonIDs( GetPolygonPerimeterVertexInstances( PolygonID ) );
+
+			for( int32 HoleIndex = 0; HoleIndex < GetNumPolygonHoles( PolygonID ); ++HoleIndex )
+			{
+				PopulatePolygonIDs( GetPolygonHoleVertexInstances( PolygonID, HoleIndex ) );
+			}
+
+			const FPolygonGroupID PolygonGroupID = PolygonArray[ PolygonID ].PolygonGroupID;
+			PolygonGroupArray[ PolygonGroupID ].Polygons.Add( PolygonID );
+
+			// We don't serialize triangles; instead the polygon gets retriangulated on load
+			ComputePolygonTriangulation( PolygonID, PolygonArray[ PolygonID ].Triangles );
+		}
+	}
 }
 
 #if WITH_EDITOR
@@ -54,6 +105,7 @@ void UMeshDescription::Empty()
 	PolygonGroupAttributesSet.Initialize( 0 );
 }
 
+#if 0
 #if WITH_EDITORONLY_DATA
 FString UMeshDescription::GetIdString()
 {
@@ -216,6 +268,7 @@ FString UMeshDescription::GetIdString()
 	return Guid.ToString(EGuidFormats::Digits);
 }
 #endif //WITH_EDITORONLY_DATA
+#endif
 
 void UMeshDescription::Compact( FElementIDRemappings& OutRemappings )
 {
diff --git a/Engine/Source/Runtime/MeshDescription/Public/MeshAttributeArray.h b/Engine/Source/Runtime/MeshDescription/Public/MeshAttributeArray.h
--- a/Engine/Source/Runtime/MeshDescription/Public/MeshAttributeArray.h
+++ b/Engine/Source/Runtime/MeshDescription/Public/MeshAttributeArray.h
@@ -28,6 +28,13 @@ using AttributeTypes = TTuple
 >;
 
 
+/**
+ * Traits class to specify which attribute types can be bulk serialized.
+ */
+template <typename T> struct TIsBulkSerializable { static const bool Value = true; };
+template <> struct TIsBulkSerializable<FName> { static const bool Value = false; };
+
+
 /**
  * This defines the container used to hold mesh element attributes of a particular name and index.
  * It is a simple TArray, so that all attributes are packed contiguously for each element ID.
@@ -40,23 +47,28 @@ class TMeshAttributeArrayBase
 {
 public:
 
-	/** Disallow use of the copy constructor to prevent arrays being mistakenly accessed in the UMeshDescription by value */
-	TMeshAttributeArrayBase( const TMeshAttributeArrayBase& ) = delete;
-
 	/**
 	 * Custom serialization for TMeshAttributeArrayBase.
 	 */
-	friend FArchive& operator<<( FArchive& Ar, TMeshAttributeArrayBase& Array )
-	{
-		Ar << Array.Container;
-		return Ar;
-	}
+	template <typename T>
+	friend typename TEnableIf<!TIsBulkSerializable<T>::Value, FArchive>::Type& operator<<( FArchive& Ar, TMeshAttributeArrayBase<T>& Array );
+
+	template <typename T>
+	friend typename TEnableIf<TIsBulkSerializable<T>::Value, FArchive>::Type& operator<<( FArchive& Ar, TMeshAttributeArrayBase<T>& Array );
 
 protected:
 
+	friend class UMeshDescription;
+
 	/** Should not instance this base class directly */
 	TMeshAttributeArrayBase() = default;
 
+	/** Disallow use of the copy constructor outside of FMeshDescription, to prevent arrays being mistakenly accessed in the UMeshDescription by value */
+	TMeshAttributeArrayBase( const TMeshAttributeArrayBase& ) = default;
+	TMeshAttributeArrayBase( TMeshAttributeArrayBase&& ) = default;
+	TMeshAttributeArrayBase& operator=( const TMeshAttributeArrayBase& ) = default;
+	TMeshAttributeArrayBase& operator=( TMeshAttributeArrayBase&& ) = default;
+
 	/** Expands the array if necessary so that the passed element index is valid. Newly created elements will be assigned the default value. */
 	void Insert( const int32 Index, const ElementType& Default )
 	{
@@ -88,6 +100,24 @@ protected:
 };
 
 
+template <typename T>
+inline typename TEnableIf<!TIsBulkSerializable<T>::Value, FArchive>::Type& operator<<( FArchive& Ar, TMeshAttributeArrayBase<T>& Array )
+{
+	// Serialize types which aren't bulk serializable, which need to be serialized element-by-element
+	Ar << Array.Container;
+	return Ar;
+}
+
+template <typename T>
+inline typename TEnableIf<TIsBulkSerializable<T>::Value, FArchive>::Type& operator<<( FArchive& Ar, TMeshAttributeArrayBase<T>& Array )
+{
+	// Serialize types which are bulk serializable, i.e. which can be memcpy'd in bulk
+	Array.Container.BulkSerialize( Ar );
+	return Ar;
+}
+
+
+
 template <typename AttributeType, typename ElementIDType>
 class TAttributeIndicesArray;
 
diff --git a/Engine/Source/Runtime/MeshDescription/Public/MeshDescription.h b/Engine/Source/Runtime/MeshDescription/Public/MeshDescription.h
--- a/Engine/Source/Runtime/MeshDescription/Public/MeshDescription.h
+++ b/Engine/Source/Runtime/MeshDescription/Public/MeshDescription.h
@@ -13,64 +13,48 @@
 #include "MeshDescription.generated.h"
 
 
-USTRUCT()
 struct FMeshVertex
 {
-	GENERATED_BODY()
-
 	FMeshVertex()
 	{}
 
 	/** All of vertex instances which reference this vertex (for split vertex support) */
-	UPROPERTY()
 	TArray<FVertexInstanceID> VertexInstanceIDs;
 
 	/** The edges connected to this vertex */
-	UPROPERTY()
 	TArray<FEdgeID> ConnectedEdgeIDs;
 
 	/** Serializer */
 	friend FArchive& operator<<( FArchive& Ar, FMeshVertex& Vertex )
 	{
-		Ar << Vertex.VertexInstanceIDs;
-		Ar << Vertex.ConnectedEdgeIDs;
 		return Ar;
 	}
 };
 
 
-USTRUCT()
 struct FMeshVertexInstance
 {
-	GENERATED_BODY()
-
 	FMeshVertexInstance()
 		: VertexID( FVertexID::Invalid )
 	{}
 
 	/** The vertex this is instancing */
-	UPROPERTY()
 	FVertexID VertexID;
 
 	/** List of connected polygons */
-	UPROPERTY()
 	TArray<FPolygonID> ConnectedPolygons;
 
 	/** Serializer */
 	friend FArchive& operator<<( FArchive& Ar, FMeshVertexInstance& VertexInstance )
 	{
 		Ar << VertexInstance.VertexID;
-		Ar << VertexInstance.ConnectedPolygons;
 		return Ar;
 	}
 };
 
 
-USTRUCT()
 struct FMeshEdge
 {
-	GENERATED_BODY()
-
 	FMeshEdge()
 	{
 		VertexIDs[ 0 ] = FVertexID::Invalid;
@@ -78,12 +62,10 @@ struct FMeshEdge
 	}
 
 	/** IDs of the two editable mesh vertices that make up this edge.  The winding direction is not defined. */
-	UPROPERTY()
 	FVertexID VertexIDs[ 2 ];
 
 	/** The polygons that share this edge.  It's best if there are always only two polygons that share
 	    the edge, and those polygons are facing the same direction */
-	UPROPERTY()
 	TArray<FPolygonID> ConnectedPolygons;
 
 	/** Serializer */
@@ -91,19 +73,14 @@ struct FMeshEdge
 	{
 		Ar << Edge.VertexIDs[ 0 ];
 		Ar << Edge.VertexIDs[ 1 ];
-		Ar << Edge.ConnectedPolygons;
 		return Ar;
 	}
 };
 
 
-USTRUCT()
 struct FMeshPolygonContour
 {
-	GENERATED_BODY()
-
 	/** The ordered list of vertex instances which make up the polygon contour. The winding direction is counter-clockwise. */
-	UPROPERTY()
 	TArray<FVertexInstanceID> VertexInstanceIDs;
 
 	/** Serializer */
@@ -153,42 +130,35 @@ struct FMeshTriangle
 	}
 
 	/** Serializer */
-	friend FArchive& operator<<( FArchive& Ar, FMeshTriangle& Tri )
+	friend FArchive& operator<<( FArchive& Ar, FMeshTriangle& Triangle )
 	{
-		Ar << Tri.VertexInstanceID0;
-		Ar << Tri.VertexInstanceID1;
-		Ar << Tri.VertexInstanceID2;
+		Ar << Triangle.VertexInstanceID0;
+		Ar << Triangle.VertexInstanceID1;
+		Ar << Triangle.VertexInstanceID2;
 		return Ar;
 	}
 };
 
 
-USTRUCT()
 struct FMeshPolygon
 {
-	GENERATED_BODY()
-
 	FMeshPolygon()
 		: PolygonGroupID( FPolygonGroupID::Invalid )
 	{}
 
 	/** The outer boundary edges of this polygon */
-	UPROPERTY()
 	FMeshPolygonContour PerimeterContour;
 
 	/** Optional inner contours of this polygon that define holes inside of the polygon.  For the geometry to
 	    be considered valid, the hole contours should reside within the boundary of the polygon perimeter contour, 
 		and must not overlap each other.  No "nesting" of polygons inside the holes is supported -- those are 
 		simply separate polygons */
-	UPROPERTY()
 	TArray<FMeshPolygonContour> HoleContours;
 
 	/** List of triangles which make up this polygon */
-	UPROPERTY()
 	TArray<FMeshTriangle> Triangles;
 
 	/** The polygon group which contains this polygon */
-	UPROPERTY()
 	FPolygonGroupID PolygonGroupID;
 
 	/** Serializer */
@@ -196,18 +166,14 @@ struct FMeshPolygon
 	{
 		Ar << Polygon.PerimeterContour;
 		Ar << Polygon.HoleContours;
-		Ar << Polygon.Triangles;
 		Ar << Polygon.PolygonGroupID;
 		return Ar;
 	}
 };
 
 
-USTRUCT()
 struct FMeshPolygonGroup
 {
-	GENERATED_BODY()
-
 	FMeshPolygonGroup()
 	{}
 
@@ -217,7 +183,6 @@ struct FMeshPolygonGroup
 	/** Serializer */
 	friend FArchive& operator<<( FArchive& Ar, FMeshPolygonGroup& PolygonGroup )
 	{
-		Ar << PolygonGroup.Polygons;
 		return Ar;
 	}
 };
@@ -240,7 +205,7 @@ enum class EComputeNTBsOptions : uint32
 };
 ENUM_CLASS_FLAGS(EComputeNTBsOptions);
 
-#define MESHDESCRIPTION_VER TEXT("8315CA62A4084F069C99E54C1C6EEDF0")
+#define MESHDESCRIPTION_VER TEXT("616CDBD71BEA4A538F74E6CAE2EACCE2")
 
 
 UCLASS()
@@ -803,8 +768,10 @@ public:
 	/** Remaps the element IDs in the mesh description according to the passed in object */
 	void Remap( const FElementIDRemappings& Remappings );
 
+#if 0
 #if WITH_EDITORONLY_DATA
 	FString GetIdString();
+#endif
 #endif
 
 	void ComputePolygonTriangulation(const FPolygonID PolygonID, TArray<FMeshTriangle>& OutTriangles);
