Copied from Dev-Editor CL 4096711

Added versioning code to UMeshDescription mesh elements so that legacy assets which still contain UMeshDescriptions still load OK.

diff --git a/Engine/Source/Runtime/Core/Public/UObject/ReleaseObjectVersion.h b/Engine/Source/Runtime/Core/Public/UObject/ReleaseObjectVersion.h
--- a/Engine/Source/Runtime/Core/Public/UObject/ReleaseObjectVersion.h
+++ b/Engine/Source/Runtime/Core/Public/UObject/ReleaseObjectVersion.h
@@ -61,6 +61,9 @@ struct CORE_API FReleaseObjectVersion
 		// During copy and paste transition getters could end up with broken state machine references
 		FixBrokenStateMachineReferencesInTransitionGetters,
 
+		// Change to MeshDescription serialization
+		MeshDescriptionNewSerialization,
+
 		// -----<new versions can be added above this line>-------------------------------------------------
 		VersionPlusOne,
 		LatestVersion = VersionPlusOne - 1
diff --git a/Engine/Source/Runtime/MeshDescription/Private/MeshDescription.cpp b/Engine/Source/Runtime/MeshDescription/Private/MeshDescription.cpp
--- a/Engine/Source/Runtime/MeshDescription/Private/MeshDescription.cpp
+++ b/Engine/Source/Runtime/MeshDescription/Private/MeshDescription.cpp
@@ -15,6 +15,9 @@ UMeshDescription::UMeshDescription()
 void UMeshDescription::Serialize( FArchive& Ar )
 {
 	Super::Serialize( Ar );
+
+	Ar.UsingCustomVersion( FReleaseObjectVersion::GUID );
+
 	Ar << VertexArray;
 	Ar << VertexInstanceArray;
 	Ar << EdgeArray;
@@ -27,7 +30,7 @@ void UMeshDescription::Serialize( FArchive& Ar )
 	Ar << PolygonAttributesSet;
 	Ar << PolygonGroupAttributesSet;
 
-	if( Ar.IsLoading() )
+	if( Ar.IsLoading() && Ar.CustomVer( FReleaseObjectVersion::GUID ) >= FReleaseObjectVersion::MeshDescriptionNewSerialization )
 	{
 		// Populate vertex instance IDs for vertices
 		for( const FVertexInstanceID VertexInstanceID : VertexInstanceArray.GetElementIDs() )
diff --git a/Engine/Source/Runtime/MeshDescription/Public/MeshDescription.h b/Engine/Source/Runtime/MeshDescription/Public/MeshDescription.h
--- a/Engine/Source/Runtime/MeshDescription/Public/MeshDescription.h
+++ b/Engine/Source/Runtime/MeshDescription/Public/MeshDescription.h
@@ -1,15 +1,14 @@
 // Copyright 1998-2018 Epic Games, Inc. All Rights Reserved.
 
 #pragma once
-#if defined(_MSC_VER) && _MSC_FULL_VER <= 190023918
-#pragma warning(disable:4503)
-#endif
 
 #include "CoreMinimal.h"
 #include "CoreTypes.h"
 #include "MeshTypes.h"
 #include "MeshElementArray.h"
 #include "MeshAttributeArray.h"
+#include "UObject/ObjectMacros.h"
+#include "UObject/ReleaseObjectVersion.h"
 #include "MeshDescription.generated.h"
 
 
@@ -27,6 +26,12 @@ struct FMeshVertex
 	/** Serializer */
 	friend FArchive& operator<<( FArchive& Ar, FMeshVertex& Vertex )
 	{
+		if( Ar.IsLoading() && Ar.CustomVer( FReleaseObjectVersion::GUID ) < FReleaseObjectVersion::MeshDescriptionNewSerialization )
+		{
+			Ar << Vertex.VertexInstanceIDs;
+			Ar << Vertex.ConnectedEdgeIDs;
+		}
+
 		return Ar;
 	}
 };
@@ -48,6 +53,11 @@ struct FMeshVertexInstance
 	friend FArchive& operator<<( FArchive& Ar, FMeshVertexInstance& VertexInstance )
 	{
 		Ar << VertexInstance.VertexID;
+		if( Ar.IsLoading() && Ar.CustomVer( FReleaseObjectVersion::GUID ) < FReleaseObjectVersion::MeshDescriptionNewSerialization )
+		{
+			Ar << VertexInstance.ConnectedPolygons;
+		}
+
 		return Ar;
 	}
 };
@@ -73,6 +83,11 @@ struct FMeshEdge
 	{
 		Ar << Edge.VertexIDs[ 0 ];
 		Ar << Edge.VertexIDs[ 1 ];
+		if( Ar.IsLoading() && Ar.CustomVer( FReleaseObjectVersion::GUID ) < FReleaseObjectVersion::MeshDescriptionNewSerialization )
+		{
+			Ar << Edge.ConnectedPolygons;
+		}
+
 		return Ar;
 	}
 };
@@ -135,6 +150,7 @@ struct FMeshTriangle
 		Ar << Triangle.VertexInstanceID0;
 		Ar << Triangle.VertexInstanceID1;
 		Ar << Triangle.VertexInstanceID2;
+
 		return Ar;
 	}
 };
@@ -166,7 +182,12 @@ struct FMeshPolygon
 	{
 		Ar << Polygon.PerimeterContour;
 		Ar << Polygon.HoleContours;
+		if( Ar.IsLoading() && Ar.CustomVer( FReleaseObjectVersion::GUID ) < FReleaseObjectVersion::MeshDescriptionNewSerialization )
+		{
+			Ar << Polygon.Triangles;
+		}
 		Ar << Polygon.PolygonGroupID;
+
 		return Ar;
 	}
 };
@@ -183,6 +204,11 @@ struct FMeshPolygonGroup
 	/** Serializer */
 	friend FArchive& operator<<( FArchive& Ar, FMeshPolygonGroup& PolygonGroup )
 	{
+		if( Ar.IsLoading() && Ar.CustomVer( FReleaseObjectVersion::GUID ) < FReleaseObjectVersion::MeshDescriptionNewSerialization )
+		{
+			Ar << PolygonGroup.Polygons;
+		}
+
 		return Ar;
 	}
 };
