Copied from Dev-Editor CL 4096832

Further versioning for legacy mesh descriptions: BulkSerialize has a different format to 'regular' serialize.

diff --git a/Engine/Source/Runtime/MeshDescription/Public/MeshAttributeArray.h b/Engine/Source/Runtime/MeshDescription/Public/MeshAttributeArray.h
--- a/Engine/Source/Runtime/MeshDescription/Public/MeshAttributeArray.h
+++ b/Engine/Source/Runtime/MeshDescription/Public/MeshAttributeArray.h
@@ -5,6 +5,7 @@
 #include "CoreMinimal.h"
 #include "MeshTypes.h"
 #include "MeshElementRemappings.h"
+#include "UObject/EditorObjectVersion.h"
 
 
 /**
@@ -111,8 +112,17 @@ inline typename TEnableIf<!TIsBulkSerializable<T>::Value, FArchive>::Type& opera
 template <typename T>
 inline typename TEnableIf<TIsBulkSerializable<T>::Value, FArchive>::Type& operator<<( FArchive& Ar, TMeshAttributeArrayBase<T>& Array )
 {
-	// Serialize types which are bulk serializable, i.e. which can be memcpy'd in bulk
-	Array.Container.BulkSerialize( Ar );
+	if( Ar.IsLoading() && Ar.CustomVer( FReleaseObjectVersion::GUID ) < FReleaseObjectVersion::MeshDescriptionNewSerialization )
+	{
+		// Legacy path for old format attribute arrays. BulkSerialize has a different format from regular serialization.
+		Ar << Array.Container;
+	}
+	else
+	{
+		// Serialize types which are bulk serializable, i.e. which can be memcpy'd in bulk
+		Array.Container.BulkSerialize( Ar );
+	}
+
 	return Ar;
 }
 
