Fixed error when merging from Dev-Editor and changing versioning code to use the Release stream.

diff --git a/Engine/Source/Runtime/MeshDescription/Public/MeshAttributeArray.h b/Engine/Source/Runtime/MeshDescription/Public/MeshAttributeArray.h
--- a/Engine/Source/Runtime/MeshDescription/Public/MeshAttributeArray.h
+++ b/Engine/Source/Runtime/MeshDescription/Public/MeshAttributeArray.h
@@ -5,7 +5,7 @@
 #include "CoreMinimal.h"
 #include "MeshTypes.h"
 #include "MeshElementRemappings.h"
-#include "UObject/EditorObjectVersion.h"
+#include "UObject/ReleaseObjectVersion.h"
 
 
 /**
