Fix ReversePolygonFacing crash

diff --git a/Engine/Source/Runtime/MeshDescription/Private/MeshDescription.cpp b/Engine/Source/Runtime/MeshDescription/Private/MeshDescription.cpp
--- a/Engine/Source/Runtime/MeshDescription/Private/MeshDescription.cpp
+++ b/Engine/Source/Runtime/MeshDescription/Private/MeshDescription.cpp
@@ -1239,20 +1239,11 @@ void UMeshDescription::GetAllCharts( TArray<TArray<FPolygonID>>& OutCharts )
 
 void UMeshDescription::ReversePolygonFacing(const FPolygonID PolygonID)
 {
-	FMeshPolygon& Polygon = GetPolygon(PolygonID);
-
 	//Build a reverse perimeter
-	TArray<FVertexInstanceID> ReverseVertexInstanceIDs;
-	int32 PerimeterReverseIndex = Polygon.PerimeterContour.VertexInstanceIDs.Num() - 1;
-	for(const FVertexInstanceID& VertexInstanceID : Polygon.PerimeterContour.VertexInstanceIDs)
-	{
-		ReverseVertexInstanceIDs[PerimeterReverseIndex] = VertexInstanceID;
-		PerimeterReverseIndex--;
-	}
-	//Assign the reverse perimeter
-	for (int32 VertexInstanceIndex = 0; VertexInstanceIndex < ReverseVertexInstanceIDs.Num(); ++VertexInstanceIndex)
+	FMeshPolygon& Polygon = GetPolygon(PolygonID);
+	for (int32 i = 0; i < Polygon.PerimeterContour.VertexInstanceIDs.Num() / 2; ++i)
 	{
-		Polygon.PerimeterContour.VertexInstanceIDs[VertexInstanceIndex] = ReverseVertexInstanceIDs[VertexInstanceIndex];
+		Polygon.PerimeterContour.VertexInstanceIDs.Swap(i, Polygon.PerimeterContour.VertexInstanceIDs.Num() - i - 1);
 	}
 	
 	//Triangulate the polygon since we reverse the indices
@@ -1264,7 +1255,6 @@ void UMeshDescription::ReverseAllPolygonFacing()
 	// Perform triangulation directly into mesh polygons
 	for (const FPolygonID PolygonID : Polygons().GetElementIDs())
 	{
-		FMeshPolygon& Polygon = PolygonArray[PolygonID];
 		ReversePolygonFacing(PolygonID);
 	}
 }
\ No newline at end of file
