Undo //UE4/Dev-Editor/Engine/Source/Runtime/MeshDescription/Public/... changelist 4096832

This will be superseded soon by changes merged over from 4.20.

diff --git a/Engine/Source/Runtime/MeshDescription/Public/MeshAttributeArray.h b/Engine/Source/Runtime/MeshDescription/Public/MeshAttributeArray.h
--- a/Engine/Source/Runtime/MeshDescription/Public/MeshAttributeArray.h
+++ b/Engine/Source/Runtime/MeshDescription/Public/MeshAttributeArray.h
@@ -5,7 +5,6 @@
 #include "CoreMinimal.h"
 #include "MeshTypes.h"
 #include "MeshElementRemappings.h"
-#include "UObject/EditorObjectVersion.h"
 
 
 /**
@@ -112,17 +111,8 @@ inline typename TEnableIf<!TIsBulkSerializable<T>::Value, FArchive>::Type& opera
 template <typename T>
 inline typename TEnableIf<TIsBulkSerializable<T>::Value, FArchive>::Type& operator<<( FArchive& Ar, TMeshAttributeArrayBase<T>& Array )
 {
-	if( Ar.IsLoading() && Ar.CustomVer( FEditorObjectVersion::GUID ) < FEditorObjectVersion::MeshDescriptionNewSerialization )
-	{
-		// Legacy path for old format attribute arrays. BulkSerialize has a different format from regular serialization.
-		Ar << Array.Container;
-	}
-	else
-	{
-		// Serialize types which are bulk serializable, i.e. which can be memcpy'd in bulk
-		Array.Container.BulkSerialize( Ar );
-	}
-
+	// Serialize types which are bulk serializable, i.e. which can be memcpy'd in bulk
+	Array.Container.BulkSerialize( Ar );
 	return Ar;
 }
 
