Undo //UE4/Dev-Editor/Engine/Source/Runtime/... changelist 4096711

This will be superseded soon by changes merged over from 4.20.

diff --git a/Engine/Source/Runtime/Core/Public/UObject/EditorObjectVersion.h b/Engine/Source/Runtime/Core/Public/UObject/EditorObjectVersion.h
--- a/Engine/Source/Runtime/Core/Public/UObject/EditorObjectVersion.h
+++ b/Engine/Source/Runtime/Core/Public/UObject/EditorObjectVersion.h
@@ -59,8 +59,6 @@ struct CORE_API FEditorObjectVersion
 		AddedMorphTargetSectionIndices,
 		// Serialize the instanced static mesh render data, to avoid building it at runtime
 		SerializeInstancedStaticMeshRenderData,
-		// Change to MeshDescription serialization
-		MeshDescriptionNewSerialization,
 		// -----<new versions can be added above this line>-------------------------------------------------
 		VersionPlusOne,
 		LatestVersion = VersionPlusOne - 1
diff --git a/Engine/Source/Runtime/MeshDescription/Private/MeshDescription.cpp b/Engine/Source/Runtime/MeshDescription/Private/MeshDescription.cpp
--- a/Engine/Source/Runtime/MeshDescription/Private/MeshDescription.cpp
+++ b/Engine/Source/Runtime/MeshDescription/Private/MeshDescription.cpp
@@ -16,8 +16,6 @@ void UDEPRECATED_MeshDescription::Serialize( FArchive& Ar )
 
 FArchive& operator<<( FArchive& Ar, FMeshDescription& MeshDescription )
 {
-	Ar.UsingCustomVersion( FEditorObjectVersion::GUID );
-
 	Ar << MeshDescription.VertexArray;
 	Ar << MeshDescription.VertexInstanceArray;
 	Ar << MeshDescription.EdgeArray;
@@ -30,7 +28,7 @@ FArchive& operator<<( FArchive& Ar, FMeshDescription& MeshDescription )
 	Ar << MeshDescription.PolygonAttributesSet;
 	Ar << MeshDescription.PolygonGroupAttributesSet;
 
-	if( Ar.IsLoading() && Ar.CustomVer( FEditorObjectVersion::GUID ) >= FEditorObjectVersion::MeshDescriptionNewSerialization )
+	if( Ar.IsLoading() )
 	{
 		// Populate vertex instance IDs for vertices
 		for( const FVertexInstanceID VertexInstanceID : MeshDescription.VertexInstanceArray.GetElementIDs() )
diff --git a/Engine/Source/Runtime/MeshDescription/Public/MeshDescription.h b/Engine/Source/Runtime/MeshDescription/Public/MeshDescription.h
--- a/Engine/Source/Runtime/MeshDescription/Public/MeshDescription.h
+++ b/Engine/Source/Runtime/MeshDescription/Public/MeshDescription.h
@@ -1,14 +1,15 @@
 // Copyright 1998-2018 Epic Games, Inc. All Rights Reserved.
 
 #pragma once
+#if defined(_MSC_VER) && _MSC_FULL_VER <= 190023918
+#pragma warning(disable:4503)
+#endif
 
 #include "CoreMinimal.h"
 #include "CoreTypes.h"
 #include "MeshTypes.h"
 #include "MeshElementArray.h"
 #include "MeshAttributeArray.h"
-#include "UObject/ObjectMacros.h"
-#include "UObject/EditorObjectVersion.h"
 #include "MeshDescription.generated.h"
 
 enum
@@ -31,12 +32,6 @@ struct FMeshVertex
 	/** Serializer */
 	friend FArchive& operator<<( FArchive& Ar, FMeshVertex& Vertex )
 	{
-		if( Ar.IsLoading() && Ar.CustomVer( FEditorObjectVersion::GUID ) < FEditorObjectVersion::MeshDescriptionNewSerialization )
-		{
-			Ar << Vertex.VertexInstanceIDs;
-			Ar << Vertex.ConnectedEdgeIDs;
-		}
-
 		return Ar;
 	}
 };
@@ -58,11 +53,6 @@ struct FMeshVertexInstance
 	friend FArchive& operator<<( FArchive& Ar, FMeshVertexInstance& VertexInstance )
 	{
 		Ar << VertexInstance.VertexID;
-		if( Ar.IsLoading() && Ar.CustomVer( FEditorObjectVersion::GUID ) < FEditorObjectVersion::MeshDescriptionNewSerialization )
-		{
-			Ar << VertexInstance.ConnectedPolygons;
-		}
-
 		return Ar;
 	}
 };
@@ -88,11 +78,6 @@ struct FMeshEdge
 	{
 		Ar << Edge.VertexIDs[ 0 ];
 		Ar << Edge.VertexIDs[ 1 ];
-		if( Ar.IsLoading() && Ar.CustomVer( FEditorObjectVersion::GUID ) < FEditorObjectVersion::MeshDescriptionNewSerialization )
-		{
-			Ar << Edge.ConnectedPolygons;
-		}
-
 		return Ar;
 	}
 };
@@ -155,7 +140,6 @@ struct FMeshTriangle
 		Ar << Triangle.VertexInstanceID0;
 		Ar << Triangle.VertexInstanceID1;
 		Ar << Triangle.VertexInstanceID2;
-
 		return Ar;
 	}
 };
@@ -187,12 +171,7 @@ struct FMeshPolygon
 	{
 		Ar << Polygon.PerimeterContour;
 		Ar << Polygon.HoleContours;
-		if( Ar.IsLoading() && Ar.CustomVer( FEditorObjectVersion::GUID ) < FEditorObjectVersion::MeshDescriptionNewSerialization )
-		{
-			Ar << Polygon.Triangles;
-		}
 		Ar << Polygon.PolygonGroupID;
-
 		return Ar;
 	}
 };
@@ -209,11 +188,6 @@ struct FMeshPolygonGroup
 	/** Serializer */
 	friend FArchive& operator<<( FArchive& Ar, FMeshPolygonGroup& PolygonGroup )
 	{
-		if( Ar.IsLoading() && Ar.CustomVer( FEditorObjectVersion::GUID ) < FEditorObjectVersion::MeshDescriptionNewSerialization )
-		{
-			Ar << PolygonGroup.Polygons;
-		}
-
 		return Ar;
 	}
 };
