New attribute array API.

Fixed some flaws in the original API, deprecated various methods, and introduced some new features.
- Now attribute arrays are accessed via TAttributesRef or TAttributesView (and corresponding const versions). These value types hold references to attribute arrays, and individual elements can be accessed by their element ID and attribute index. Using a value type is safer than the previous method which required assignment to a const-ref (and not doing so would take a temporary copy of the attribute array).
- The attribute set has been totally flattened, so all attributes of different types are added to the same container. This greatly improves compile times, prevents attributes from being created with the same name but different types, and permits the view feature.
- The class hierarchy has changed to have generic base classes where possible with no particular ElementID type. This reduces the code footprint by no longer generating nearly identical copies of templated methods.
- A TAttributesView allows the user to access an attribute array by the type of their choosing, regardless of its actual type. For example, the Normal attribute may be registered with type FPackedVector, but accessed as if it was an FVector. This allows us to move away from very strong typing, and instead transparently store attributes of a more efficient size, while the user is not affected.
- A transient attribute flag has been added, to denote that a particular attribute should not be saved.

diff --git a/Engine/Plugins/Editor/MeshEditor/Source/MeshEditor/MeshEditorMode.cpp b/Engine/Plugins/Editor/MeshEditor/Source/MeshEditor/MeshEditorMode.cpp
--- a/Engine/Plugins/Editor/MeshEditor/Source/MeshEditor/MeshEditorMode.cpp
+++ b/Engine/Plugins/Editor/MeshEditor/Source/MeshEditor/MeshEditorMode.cpp
@@ -1363,9 +1363,9 @@ void FMeshEditorMode::UpdateDebugNormals()
 		const UEditableMesh* EditableMesh = ComponentAndEditableMesh.EditableMesh;
 		const FMeshDescription* MeshDescription = EditableMesh->GetMeshDescription();
 
-		const TVertexAttributeArray<FVector>& VertexPositions = MeshDescription->VertexAttributes().GetAttributes<FVector>( MeshAttribute::Vertex::Position );
-		const TVertexInstanceAttributeArray<FVector>& VertexNormals = MeshDescription->VertexInstanceAttributes().GetAttributes<FVector>( MeshAttribute::VertexInstance::Normal );
-		const TVertexInstanceAttributeArray<FVector>& VertexTangents = MeshDescription->VertexInstanceAttributes().GetAttributes<FVector>( MeshAttribute::VertexInstance::Tangent );
+		TVertexAttributesConstRef<FVector> VertexPositions = MeshDescription->VertexAttributes().GetAttributesRef<FVector>( MeshAttribute::Vertex::Position );
+		TVertexInstanceAttributesConstRef<FVector> VertexNormals = MeshDescription->VertexInstanceAttributes().GetAttributesRef<FVector>( MeshAttribute::VertexInstance::Normal );
+		TVertexInstanceAttributesConstRef<FVector> VertexTangents = MeshDescription->VertexInstanceAttributes().GetAttributesRef<FVector>( MeshAttribute::VertexInstance::Tangent );
 
 		const FMatrix ComponentToWorldMatrix = Component->GetRenderMatrix();
 
@@ -1956,7 +1956,7 @@ void FMeshEditorMode::FrameSelectedElements( FEditorViewportClient* ViewportClie
 				const TArray<FMeshElement>& VertexElements = SelectedMeshAndVertices.Value;
 
 				const FMeshDescription* MeshDescription = EditableMesh->GetMeshDescription();
-				const TVertexAttributeArray<FVector>& VertexPositions = MeshDescription->VertexAttributes().GetAttributes<FVector>( MeshAttribute::Vertex::Position );
+				TVertexAttributesConstRef<FVector> VertexPositions = MeshDescription->VertexAttributes().GetAttributesRef<FVector>( MeshAttribute::Vertex::Position );
 
 				for( const auto& VertexElement : VertexElements )
 				{
@@ -1983,7 +1983,7 @@ void FMeshEditorMode::FrameSelectedElements( FEditorViewportClient* ViewportClie
 				const TArray<FMeshElement>& EdgeElements = SelectedMeshAndEdges.Value;
 
 				const FMeshDescription* MeshDescription = EditableMesh->GetMeshDescription();
-				const TVertexAttributeArray<FVector>& VertexPositions = MeshDescription->VertexAttributes().GetAttributes<FVector>( MeshAttribute::Vertex::Position );
+				TVertexAttributesConstRef<FVector> VertexPositions = MeshDescription->VertexAttributes().GetAttributesRef<FVector>( MeshAttribute::Vertex::Position );
 
 				for( const auto& EdgeElement : EdgeElements )
 				{
@@ -2015,7 +2015,7 @@ void FMeshEditorMode::FrameSelectedElements( FEditorViewportClient* ViewportClie
 				const TArray<FMeshElement>& PolygonElements = SelectedMeshAndPolygons.Value;
 
 				const FMeshDescription* MeshDescription = EditableMesh->GetMeshDescription();
-				const TVertexAttributeArray<FVector>& VertexPositions = MeshDescription->VertexAttributes().GetAttributes<FVector>( MeshAttribute::Vertex::Position );
+				TVertexAttributesConstRef<FVector> VertexPositions = MeshDescription->VertexAttributes().GetAttributesRef<FVector>( MeshAttribute::Vertex::Position );
 
 				for( const auto& PolygonElement : PolygonElements )
 				{
@@ -2371,7 +2371,7 @@ void FMeshEditorMode::AddMeshElementToOverlay( UOverlayComponent* OverlayCompone
 			if( MeshElement.IsElementIDValid( EditableMesh ) )
 			{
 				const FMeshDescription* MeshDescription = EditableMesh->GetMeshDescription();
-				const TVertexAttributeArray<FVector>& VertexPositions = MeshDescription->VertexAttributes().GetAttributes<FVector>( MeshAttribute::Vertex::Position );
+				TVertexAttributesConstRef<FVector> VertexPositions = MeshDescription->VertexAttributes().GetAttributesRef<FVector>( MeshAttribute::Vertex::Position );
 
 				UPrimitiveComponent* Component = MeshElement.Component.Get();
 				check( Component != nullptr );
@@ -3249,7 +3249,7 @@ void FMeshEditorMode::UpdateActiveAction( const bool bIsActionFinishing )
 
 			UEditableMesh* EditableMesh = MeshAndTransformables.Key;
 			const FMeshDescription* MeshDescription = EditableMesh->GetMeshDescription();
-			const TVertexAttributeArray<FVector>& VertexPositions = MeshDescription->VertexAttributes().GetAttributes<FVector>( MeshAttribute::Vertex::Position );
+			TVertexAttributesConstRef<FVector> VertexPositions = MeshDescription->VertexAttributes().GetAttributesRef<FVector>( MeshAttribute::Vertex::Position );
 
 			const TArray< const FMeshElementViewportTransformable* >& TransformablesForMesh = MeshAndTransformables.Value;
 
@@ -3402,7 +3402,7 @@ bool FMeshEditorMode::FindEdgeSplitUnderInteractor(	UViewportInteractor* Viewpor
 	bool bFoundSplit = false;
 
 	const FMeshDescription* MeshDescription = EditableMesh->GetMeshDescription();
-	const TVertexAttributeArray<FVector>& VertexPositions = MeshDescription->VertexAttributes().GetAttributes<FVector>( MeshAttribute::Vertex::Position );
+	TVertexAttributesConstRef<FVector> VertexPositions = MeshDescription->VertexAttributes().GetAttributesRef<FVector>( MeshAttribute::Vertex::Position );
 
 	// Figure out where to split based on where the interactor is aiming.  We'll look at all of the
 	// selected edges, and choose a split offset based on the closest point along one of those edges
@@ -3508,7 +3508,7 @@ FEditableMeshElementAddress FMeshEditorMode::QueryElement( const UEditableMesh&
 	static TSet<FPolygonID> FrontFacingPolygons;
 	FrontFacingPolygons.Reset();
 
-	const TPolygonAttributeArray<FVector>& PolygonCenters = EditableMesh.GetMeshDescription()->PolygonAttributes().GetAttributes<FVector>( MeshAttribute::Polygon::Center );
+	TPolygonAttributesConstRef<FVector> PolygonCenters = EditableMesh.GetMeshDescription()->PolygonAttributes().GetAttributesRef<FVector>( MeshAttribute::Polygon::Center );
 
 	// Look for all the front-facing elements
 	for( const FPolygonID PolygonID : CandidatePolygons )
@@ -3538,7 +3538,7 @@ FEditableMeshElementAddress FMeshEditorMode::QueryElement( const UEditableMesh&
 	FVector CurrentRayEnd = RayEnd;
 
 	const FMeshDescription* MeshDescription = EditableMesh.GetMeshDescription();
-	const TVertexAttributeArray<FVector>& VertexPositions = MeshDescription->VertexAttributes().GetAttributes<FVector>( MeshAttribute::Vertex::Position );
+	TVertexAttributesConstRef<FVector> VertexPositions = MeshDescription->VertexAttributes().GetAttributesRef<FVector>( MeshAttribute::Vertex::Position );
 
 	// Check polygons first; this is so we always impose a closest hit location at the poly before checking other elements, so anything behind is occluded
 	for( const FPolygonID PolygonID : FrontFacingPolygons )
@@ -4345,7 +4345,7 @@ bool FMeshEditorMode::FrustumSelect( const FConvexVolume& InFrustum, FEditorView
 		FTransform ComponentTransform = Component->GetComponentTransform();
 
 		const FMeshDescription* MeshDescription = EditableMesh->GetMeshDescription();
-		const TVertexAttributeArray<FVector>& VertexPositions = MeshDescription->VertexAttributes().GetAttributes<FVector>( MeshAttribute::Vertex::Position );
+		TVertexAttributesConstRef<FVector> VertexPositions = MeshDescription->VertexAttributes().GetAttributesRef<FVector>( MeshAttribute::Vertex::Position );
 
 		static TArray<FEdgeID> SelectedEdgeIDs;
 		SelectedEdgeIDs.Empty();
@@ -4575,7 +4575,7 @@ void FMeshEditorMode::RefreshTransformables( const bool bNewObjectsSelected )
 					const FMatrix ComponentToWorldMatrix = Component->GetRenderMatrix();
 
 					const FMeshDescription* MeshDescription = EditableMesh->GetMeshDescription();
-					const TVertexAttributeArray<FVector>& VertexPositions = MeshDescription->VertexAttributes().GetAttributes<FVector>( MeshAttribute::Vertex::Position );
+					TVertexAttributesConstRef<FVector> VertexPositions = MeshDescription->VertexAttributes().GetAttributesRef<FVector>( MeshAttribute::Vertex::Position );
 
 					FTransform ElementTransform = FTransform::Identity;
 					switch( MeshElement.ElementAddress.ElementType )
@@ -4607,7 +4607,8 @@ void FMeshEditorMode::RefreshTransformables( const bool bNewObjectsSelected )
 						{
 							const FPolygonID PolygonID( MeshElement.ElementAddress.ElementID );
 
-							const TPolygonAttributeArray<FVector>& PolygonCenters = EditableMesh->GetMeshDescription()->PolygonAttributes().GetAttributes<FVector>( MeshAttribute::Polygon::Center );
+							TPolygonAttributesConstRef<FVector> PolygonCenters = EditableMesh->GetMeshDescription()->PolygonAttributes().GetAttributesRef<FVector>( MeshAttribute::Polygon::Center );
+
 							const FVector ComponentSpacePolygonCenter = PolygonCenters[ PolygonID ];
 							ElementTransform.SetLocation( ComponentToWorldMatrix.TransformPosition( ComponentSpacePolygonCenter ) );
 
diff --git a/Engine/Plugins/Editor/MeshEditor/Source/MeshEditor/MeshEditorStaticMeshAdapter.cpp b/Engine/Plugins/Editor/MeshEditor/Source/MeshEditor/MeshEditorStaticMeshAdapter.cpp
--- a/Engine/Plugins/Editor/MeshEditor/Source/MeshEditor/MeshEditorStaticMeshAdapter.cpp
+++ b/Engine/Plugins/Editor/MeshEditor/Source/MeshEditor/MeshEditorStaticMeshAdapter.cpp
@@ -59,7 +59,7 @@ void UMeshEditorStaticMeshAdapter::OnRebuildRenderMesh( const UEditableMesh* Edi
 	check( MeshDescription );
 
 	// Add all vertices
-	const TVertexAttributeArray<FVector>& VertexPositions = MeshDescription->VertexAttributes().GetAttributes<FVector>( MeshAttribute::Vertex::Position );
+	TVertexAttributesConstRef<FVector> VertexPositions = MeshDescription->VertexAttributes().GetAttributesRef<FVector>( MeshAttribute::Vertex::Position );
 	for( const FVertexID VertexID : MeshDescription->Vertices().GetElementIDs() )
 	{
 		WireframeMesh->AddVertex( VertexID );
@@ -67,8 +67,8 @@ void UMeshEditorStaticMeshAdapter::OnRebuildRenderMesh( const UEditableMesh* Edi
 	}
 
 	// Add all edges
-	const TEdgeAttributeArray<bool>& EdgeHardnesses = MeshDescription->EdgeAttributes().GetAttributes<bool>( MeshAttribute::Edge::IsHard );
-	const TEdgeAttributeArray<bool>& EdgeUVSeams = MeshDescription->EdgeAttributes().GetAttributes<bool>( MeshAttribute::Edge::IsUVSeam );
+	TEdgeAttributesConstRef<bool> EdgeHardnesses = MeshDescription->EdgeAttributes().GetAttributesRef<bool>( MeshAttribute::Edge::IsHard );
+	TEdgeAttributesConstRef<bool> EdgeUVSeams = MeshDescription->EdgeAttributes().GetAttributesRef<bool>( MeshAttribute::Edge::IsUVSeam );
 	for( const FEdgeID EdgeID : MeshDescription->Edges().GetElementIDs() )
 	{
 		WireframeMesh->AddEdge( EdgeID );
@@ -78,7 +78,7 @@ void UMeshEditorStaticMeshAdapter::OnRebuildRenderMesh( const UEditableMesh* Edi
 
 	// Add all polygons and edge instances
 	// @todo mesheditor wireframe: holes not yet handled
-	const TPolygonAttributeArray<FVector>& PolygonNormals = MeshDescription->PolygonAttributes().GetAttributes<FVector>( MeshAttribute::Polygon::Normal );
+	TPolygonAttributesConstRef<FVector> PolygonNormals = MeshDescription->PolygonAttributes().GetAttributesRef<FVector>( MeshAttribute::Polygon::Normal );
 	for( const FPolygonID PolygonID : EditableMesh->GetMeshDescription()->Polygons().GetElementIDs() )
 	{
 		WireframeMesh->AddPolygon( PolygonID );
@@ -99,7 +99,7 @@ void UMeshEditorStaticMeshAdapter::OnRebuildRenderMeshFinish( const UEditableMes
 	const FMeshDescription* MeshDescription = EditableMesh->GetMeshDescription();
 	check( MeshDescription );
 
-	const TPolygonAttributeArray<FVector>& PolygonNormals = MeshDescription->PolygonAttributes().GetAttributes<FVector>( MeshAttribute::Polygon::Normal );
+	TPolygonAttributesConstRef<FVector> PolygonNormals = MeshDescription->PolygonAttributes().GetAttributesRef<FVector>( MeshAttribute::Polygon::Normal );
 	for( const FPolygonID PolygonID : EditableMesh->PolygonsPendingNewTangentBasis )
 	{
 		WireframeMesh->SetPolygonNormal( PolygonID, PolygonNormals[ PolygonID ] );
@@ -202,7 +202,7 @@ void UMeshEditorStaticMeshAdapter::OnCreateVertices( const UEditableMesh* Editab
 	const FMeshDescription* MeshDescription = EditableMesh->GetMeshDescription();
 	check( MeshDescription );
 
-	const TVertexAttributeArray<FVector>& VertexPositions = MeshDescription->VertexAttributes().GetAttributes<FVector>( MeshAttribute::Vertex::Position );
+	TVertexAttributesConstRef<FVector> VertexPositions = MeshDescription->VertexAttributes().GetAttributesRef<FVector>( MeshAttribute::Vertex::Position );
 	for( const FVertexID VertexID : VertexIDs )
 	{
 		WireframeMesh->AddVertex( VertexID );
@@ -230,8 +230,8 @@ void UMeshEditorStaticMeshAdapter::OnCreateEdges( const UEditableMesh* EditableM
 	const FMeshDescription* MeshDescription = EditableMesh->GetMeshDescription();
 	check( MeshDescription );
 
-	const TEdgeAttributeArray<bool>& EdgeHardnesses = MeshDescription->EdgeAttributes().GetAttributes<bool>( MeshAttribute::Edge::IsHard );
-	const TEdgeAttributeArray<bool>& EdgeUVSeams = MeshDescription->EdgeAttributes().GetAttributes<bool>( MeshAttribute::Edge::IsUVSeam );
+	TEdgeAttributesConstRef<bool> EdgeHardnesses = MeshDescription->EdgeAttributes().GetAttributesRef<bool>( MeshAttribute::Edge::IsHard );
+	TEdgeAttributesConstRef<bool> EdgeUVSeams = MeshDescription->EdgeAttributes().GetAttributesRef<bool>( MeshAttribute::Edge::IsUVSeam );
 
 	for( const FEdgeID EdgeID : EdgeIDs )
 	{
@@ -278,8 +278,8 @@ void UMeshEditorStaticMeshAdapter::OnSetEdgeAttribute( const UEditableMesh* Edit
 	const FMeshDescription* MeshDescription = EditableMesh->GetMeshDescription();
 	check( MeshDescription );
 
-	const TEdgeAttributeArray<bool>& EdgeHardnesses = MeshDescription->EdgeAttributes().GetAttributes<bool>( MeshAttribute::Edge::IsHard );
-	const TEdgeAttributeArray<bool>& EdgeUVSeams = MeshDescription->EdgeAttributes().GetAttributes<bool>( MeshAttribute::Edge::IsUVSeam );
+	TEdgeAttributesConstRef<bool> EdgeHardnesses = MeshDescription->EdgeAttributes().GetAttributesRef<bool>( MeshAttribute::Edge::IsHard );
+	TEdgeAttributesConstRef<bool> EdgeUVSeams = MeshDescription->EdgeAttributes().GetAttributesRef<bool>( MeshAttribute::Edge::IsUVSeam );
 
 	if( Attribute.AttributeName == MeshAttribute::Edge::IsHard ||
 		Attribute.AttributeName == MeshAttribute::Edge::IsUVSeam )
@@ -294,7 +294,7 @@ void UMeshEditorStaticMeshAdapter::OnCreatePolygons( const UEditableMesh* Editab
 	const FMeshDescription* MeshDescription = EditableMesh->GetMeshDescription();
 	check( MeshDescription );
 
-	const TPolygonAttributeArray<FVector>& PolygonNormals = MeshDescription->PolygonAttributes().GetAttributes<FVector>( MeshAttribute::Polygon::Normal );
+	TPolygonAttributesConstRef<FVector> PolygonNormals = MeshDescription->PolygonAttributes().GetAttributesRef<FVector>( MeshAttribute::Polygon::Normal );
 	for( const FPolygonID PolygonID : PolygonIDs )
 	{
 		WireframeMesh->AddPolygon( PolygonID );
diff --git a/Engine/Plugins/Editor/MeshEditor/Source/MeshEditor/MeshEditorUtilities.cpp b/Engine/Plugins/Editor/MeshEditor/Source/MeshEditor/MeshEditorUtilities.cpp
--- a/Engine/Plugins/Editor/MeshEditor/Source/MeshEditor/MeshEditorUtilities.cpp
+++ b/Engine/Plugins/Editor/MeshEditor/Source/MeshEditor/MeshEditorUtilities.cpp
@@ -29,7 +29,7 @@ bool FMeshEditorUtilities::AssignMaterialToPolygons( UMaterialInterface* Selecte
 
 		if ( SelectedMaterial != nullptr )
 		{
-			const TPolygonGroupAttributeArray<FName>& PolygonGroupMaterialAssetNames = MeshDescription->PolygonGroupAttributes().GetAttributes<FName>( MeshAttribute::PolygonGroup::MaterialAssetName );
+			TPolygonGroupAttributesConstRef<FName> PolygonGroupMaterialAssetNames = MeshDescription->PolygonGroupAttributes().GetAttributesRef<FName>( MeshAttribute::PolygonGroup::MaterialAssetName );
 
 			for( const FPolygonGroupID PolygonGroupID : MeshDescription->PolygonGroups().GetElementIDs() )
 			{
@@ -57,7 +57,7 @@ bool FMeshEditorUtilities::AssignMaterialToPolygons( UMaterialInterface* Selecte
 			// and adding a unique suffix if there are other polygon groups with the same material slot name.
 			auto MakeUniqueSlotName = [ MeshDescription ]( FName Name ) -> FName
 			{
-				const TPolygonGroupAttributeArray<FName>& MaterialSlotNames = MeshDescription->PolygonGroupAttributes().GetAttributes<FName>( MeshAttribute::PolygonGroup::MaterialAssetName );
+				TPolygonGroupAttributesConstRef<FName> MaterialSlotNames = MeshDescription->PolygonGroupAttributes().GetAttributesRef<FName>( MeshAttribute::PolygonGroup::MaterialAssetName );
 				for( const FPolygonGroupID PolygonGroupID : MeshDescription->PolygonGroups().GetElementIDs() )
 				{
 					const FName ExistingName = MaterialSlotNames[ PolygonGroupID ];
diff --git a/Engine/Plugins/Editor/MeshEditor/Source/MeshEditor/MeshElementViewportTransformable.cpp b/Engine/Plugins/Editor/MeshEditor/Source/MeshEditor/MeshElementViewportTransformable.cpp
--- a/Engine/Plugins/Editor/MeshEditor/Source/MeshEditor/MeshElementViewportTransformable.cpp
+++ b/Engine/Plugins/Editor/MeshEditor/Source/MeshEditor/MeshElementViewportTransformable.cpp
@@ -35,12 +35,12 @@ FBox FMeshElementViewportTransformable::BuildBoundingBox( const FTransform& Boun
 		const FTransform& ComponentToWorld = MeshElement.Component.Get()->GetComponentToWorld();
 		const FTransform ComponentToBoundingBox = ComponentToWorld * WorldToBoundingBox;
 
-		UEditableMesh* EditableMesh = MeshEditorMode.FindOrCreateEditableMesh( *MeshElement.Component, MeshElement.ElementAddress.SubMeshAddress );
+		const UEditableMesh* EditableMesh = MeshEditorMode.FindOrCreateEditableMesh( *MeshElement.Component, MeshElement.ElementAddress.SubMeshAddress );
 		if( EditableMesh != nullptr )
 		{
 			if( MeshElement.IsElementIDValid( EditableMesh ) )
 			{
-				const TVertexAttributeArray<FVector>& VertexPositions = EditableMesh->GetMeshDescription()->VertexAttributes().GetAttributes<FVector>( MeshAttribute::Vertex::Position );
+				TVertexAttributesConstRef<FVector> VertexPositions = EditableMesh->GetMeshDescription()->VertexAttributes().GetAttributesRef<FVector>( MeshAttribute::Vertex::Position );
 				BoundingBox.Init();
 
 				FTransform ElementTransform = FTransform::Identity;
diff --git a/Engine/Plugins/Editor/MeshEditor/Source/PolygonModeling/BevelOrInsetPolygon.cpp b/Engine/Plugins/Editor/MeshEditor/Source/PolygonModeling/BevelOrInsetPolygon.cpp
--- a/Engine/Plugins/Editor/MeshEditor/Source/PolygonModeling/BevelOrInsetPolygon.cpp
+++ b/Engine/Plugins/Editor/MeshEditor/Source/PolygonModeling/BevelOrInsetPolygon.cpp
@@ -67,7 +67,7 @@ namespace BevelOrInsetPolygonHelpers
 				float CenterWeight = 0.0f;
 				FVector BestEdgeVertex0Position, BestEdgeVertex1Position;
 				{
-					const TVertexAttributeArray<FVector>& VertexPositions = EditableMesh->GetMeshDescription()->VertexAttributes().GetAttributes<FVector>( MeshAttribute::Vertex::Position );
+					TVertexAttributesConstRef<FVector> VertexPositions = EditableMesh->GetMeshDescription()->VertexAttributes().GetAttributesRef<FVector>( MeshAttribute::Vertex::Position );
 
 					static TArray<FVertexID> PerimeterVertexIDs;
 					EditableMesh->GetPolygonPerimeterVertices( PolygonID, /* Out */ PerimeterVertexIDs );
diff --git a/Engine/Plugins/Editor/MeshEditor/Source/PolygonModeling/EditVertexOrEdgeSharpness.cpp b/Engine/Plugins/Editor/MeshEditor/Source/PolygonModeling/EditVertexOrEdgeSharpness.cpp
--- a/Engine/Plugins/Editor/MeshEditor/Source/PolygonModeling/EditVertexOrEdgeSharpness.cpp
+++ b/Engine/Plugins/Editor/MeshEditor/Source/PolygonModeling/EditVertexOrEdgeSharpness.cpp
@@ -113,7 +113,7 @@ void UEditVertexCornerSharpnessCommand::ApplyDuringDrag( IMeshEditorModeEditingC
 			static TArray<float> NewSharpnessValues;
 			NewSharpnessValues.Reset();
 
-			const TVertexAttributeArray<float>& VertexSharpnesses = EditableMesh->GetMeshDescription()->VertexAttributes().GetAttributes<float>( MeshAttribute::Vertex::CornerSharpness );
+			const TVertexAttributesRef<float> VertexSharpnesses = EditableMesh->GetMeshDescription()->VertexAttributes().GetAttributesRef<float>( MeshAttribute::Vertex::CornerSharpness );
 
 			for( const FMeshElement& VertexElement : VertexElements )
 			{
@@ -164,7 +164,7 @@ void UEditEdgeCreaseSharpnessCommand::ApplyDuringDrag( IMeshEditorModeEditingCon
 			UEditableMesh* EditableMesh = MeshAndSelectedEdges.Key;
 			const TArray<FMeshElement>& EdgeElements = MeshAndSelectedEdges.Value;
 
-			const TEdgeAttributeArray<float>& EdgeSharpnesses = EditableMesh->GetMeshDescription()->EdgeAttributes().GetAttributes<float>( MeshAttribute::Edge::CreaseSharpness );
+			const TEdgeAttributesRef<float> EdgeSharpnesses = EditableMesh->GetMeshDescription()->EdgeAttributes().GetAttributesRef<float>( MeshAttribute::Edge::CreaseSharpness );
 
 			static TArray<FEdgeID> EdgeIDs;
 			EdgeIDs.Reset();
diff --git a/Engine/Plugins/Editor/MeshEditor/Source/PolygonModeling/SplitPolygon.cpp b/Engine/Plugins/Editor/MeshEditor/Source/PolygonModeling/SplitPolygon.cpp
--- a/Engine/Plugins/Editor/MeshEditor/Source/PolygonModeling/SplitPolygon.cpp
+++ b/Engine/Plugins/Editor/MeshEditor/Source/PolygonModeling/SplitPolygon.cpp
@@ -134,7 +134,7 @@ void USplitPolygonCommand::ApplyDuringDrag( IMeshEditorModeEditingContract& Mesh
 			FVector LaserPointerStart, LaserPointerEnd;
 			if( ViewportInteractor->GetLaserPointer( /* Out */ LaserPointerStart, /* Out */ LaserPointerEnd ) )
 			{
-				const TVertexAttributeArray<FVector>& VertexPositions = EditableMesh->GetMeshDescription()->VertexAttributes().GetAttributes<FVector>( MeshAttribute::Vertex::Position );
+				const TVertexAttributesRef<FVector> VertexPositions = EditableMesh->GetMeshDescription()->VertexAttributes().GetAttributesRef<FVector>( MeshAttribute::Vertex::Position );
 
 				for( const FPolygonID CandidatePolygonID : CandidatePolygonIDs )
 				{
diff --git a/Engine/Plugins/Editor/SpeedTreeImporter/Source/SpeedTreeImporter/Private/SpeedTreeImportFactory.cpp b/Engine/Plugins/Editor/SpeedTreeImporter/Source/SpeedTreeImporter/Private/SpeedTreeImportFactory.cpp
--- a/Engine/Plugins/Editor/SpeedTreeImporter/Source/SpeedTreeImporter/Private/SpeedTreeImportFactory.cpp
+++ b/Engine/Plugins/Editor/SpeedTreeImporter/Source/SpeedTreeImporter/Private/SpeedTreeImportFactory.cpp
@@ -1367,11 +1367,11 @@ FVertexInstanceID ProcessTriangleCorner(
 	const int32 IndexOffset,
 	const int32 NumUVs,
 	const SpeedTree::SRenderState* RenderState,
-	TVertexInstanceAttributeArray<FVector>& VertexInstanceNormals,
-	TVertexInstanceAttributeArray<FVector>& VertexInstanceTangents,
-	TVertexInstanceAttributeArray<float>& VertexInstanceBinormalSigns,
-	TVertexInstanceAttributeArray<FVector4>& VertexInstanceColors,
-	TVertexInstanceAttributeIndicesArray<FVector2D>& VertexInstanceUVs)
+	TVertexInstanceAttributesRef<FVector> VertexInstanceNormals,
+	TVertexInstanceAttributesRef<FVector> VertexInstanceTangents,
+	TVertexInstanceAttributesRef<float> VertexInstanceBinormalSigns,
+	TVertexInstanceAttributesRef<FVector4> VertexInstanceColors,
+	TVertexInstanceAttributesRef<FVector2D> VertexInstanceUVs)
 {
 	//Speedtree use 8 UVs to store is data
 	check(VertexInstanceUVs.GetNumIndices() == 8);
@@ -1402,7 +1402,7 @@ FVertexInstanceID ProcessTriangleCorner(
 	// keep texcoords padded to align indices
 	for( int32 PadIndex = 0; PadIndex < NumUVs; ++PadIndex )
 	{
-		VertexInstanceUVs.GetArrayForIndex(PadIndex)[VertexInstanceID] = FVector2D(0.0f, 0.0f);
+		VertexInstanceUVs.Set(VertexInstanceID, PadIndex, FVector2D(0.0f, 0.0f));
 	}
 
 	// All texcoords are packed into 4 float4 vertex attributes
@@ -1422,28 +1422,28 @@ FVertexInstanceID ProcessTriangleCorner(
 
 	// diffuse
 	DrawCall->GetProperty( SpeedTree::VERTEX_PROPERTY_DIFFUSE_TEXCOORDS, VertexIndex, Data );
-	VertexInstanceUVs.GetArrayForIndex(0)[VertexInstanceID] = FVector2D( Data[ 0 ], Data[ 1 ] );
+	VertexInstanceUVs.Set( VertexInstanceID, 0, FVector2D( Data[ 0 ], Data[ 1 ] ) );
 
 	// lightmap
 	DrawCall->GetProperty( SpeedTree::VERTEX_PROPERTY_LIGHTMAP_TEXCOORDS, VertexIndex, Data );
-	VertexInstanceUVs.GetArrayForIndex(1)[VertexInstanceID] = FVector2D( Data[ 0 ], Data[ 1 ] );
+	VertexInstanceUVs.Set( VertexInstanceID, 1, FVector2D( Data[ 0 ], Data[ 1 ] ) );
 
 	// branch wind
 	DrawCall->GetProperty( SpeedTree::VERTEX_PROPERTY_WIND_BRANCH_DATA, VertexIndex, Data );
-	VertexInstanceUVs.GetArrayForIndex(2)[VertexInstanceID] = FVector2D( Data[ 0 ], Data[ 1 ] );
+	VertexInstanceUVs.Set( VertexInstanceID, 2, FVector2D( Data[ 0 ], Data[ 1 ] ) );
 
 	// lod
 	if( RenderState->m_bFacingLeavesPresent )
 	{
 		DrawCall->GetProperty( SpeedTree::VERTEX_PROPERTY_LEAF_CARD_LOD_SCALAR, VertexIndex, Data );
-		VertexInstanceUVs.GetArrayForIndex(3)[VertexInstanceID] = FVector2D( Data[ 0 ], 0.0f );
-		VertexInstanceUVs.GetArrayForIndex(4)[VertexInstanceID] = FVector2D( 0.0f, 0.0f );
+		VertexInstanceUVs.Set( VertexInstanceID, 3, FVector2D( Data[ 0 ], 0.0f ) );
+		VertexInstanceUVs.Set( VertexInstanceID, 4, FVector2D( 0.0f, 0.0f ) );
 	}
 	else
 	{
 		DrawCall->GetProperty( SpeedTree::VERTEX_PROPERTY_LOD_POSITION, VertexIndex, Data );
-		VertexInstanceUVs.GetArrayForIndex(3)[VertexInstanceID] = FVector2D( -Data[ 0 ], Data[ 1 ] );
-		VertexInstanceUVs.GetArrayForIndex(4)[VertexInstanceID] = FVector2D( Data[ 2 ], 0.0f );
+		VertexInstanceUVs.Set( VertexInstanceID, 3, FVector2D( -Data[ 0 ], Data[ 1 ] ) );
+		VertexInstanceUVs.Set( VertexInstanceID, 4, FVector2D( Data[ 2 ], 0.0f ) );
 	}
 
 	// other
@@ -1451,19 +1451,19 @@ FVertexInstanceID ProcessTriangleCorner(
 	{
 		// detail
 		DrawCall->GetProperty( SpeedTree::VERTEX_PROPERTY_DETAIL_TEXCOORDS, VertexIndex, Data );
-		VertexInstanceUVs.GetArrayForIndex(5)[VertexInstanceID] = FVector2D( Data[ 0 ], Data[ 1 ] );
+		VertexInstanceUVs.Set( VertexInstanceID, 5, FVector2D( Data[ 0 ], Data[ 1 ] ) );
 
 		// branch seam
 		DrawCall->GetProperty( SpeedTree::VERTEX_PROPERTY_BRANCH_SEAM_DIFFUSE, VertexIndex, Data );
-		VertexInstanceUVs.GetArrayForIndex(6)[VertexInstanceID] = FVector2D( Data[ 0 ], Data[ 1 ] );
-		VertexInstanceUVs.GetArrayForIndex(4)[VertexInstanceID].Y = Data[ 2 ];
+		VertexInstanceUVs.Set( VertexInstanceID, 6, FVector2D( Data[ 0 ], Data[ 1 ] ) );
+		VertexInstanceUVs.Set( VertexInstanceID, 4, FVector2D( VertexInstanceUVs.Get( VertexInstanceID, 4 ).X, Data[ 2 ] ) );
 	}
 	else if( RenderState->m_bFrondsPresent )
 	{
 		// frond wind
 		DrawCall->GetProperty( SpeedTree::VERTEX_PROPERTY_WIND_EXTRA_DATA, VertexIndex, Data );
-		VertexInstanceUVs.GetArrayForIndex(5)[VertexInstanceID] = FVector2D( Data[ 0 ], Data[ 1 ] );
-		VertexInstanceUVs.GetArrayForIndex(6)[VertexInstanceID] = FVector2D( Data[ 2 ], 0.0f );
+		VertexInstanceUVs.Set( VertexInstanceID, 5, FVector2D( Data[ 0 ], Data[ 1 ] ) );
+		VertexInstanceUVs.Set( VertexInstanceID, 6, FVector2D( Data[ 2 ], 0.0f ) );
 	}
 	else if( RenderState->m_bLeavesPresent || RenderState->m_bFacingLeavesPresent )
 	{
@@ -1476,15 +1476,15 @@ FVertexInstanceID ProcessTriangleCorner(
 		{
 			DrawCall->GetProperty( SpeedTree::VERTEX_PROPERTY_LEAF_ANCHOR_POINT, VertexIndex, Data );
 		}
-		VertexInstanceUVs.GetArrayForIndex(4)[VertexInstanceID].Y = -Data[ 0 ];
-		VertexInstanceUVs.GetArrayForIndex(5)[VertexInstanceID] = FVector2D( Data[ 1 ], Data[ 2 ] );
+		VertexInstanceUVs.Set( VertexInstanceID, 4, FVector2D( VertexInstanceUVs.Get( VertexInstanceID, 4 ).X, -Data[ 0 ] ) );
+		VertexInstanceUVs.Set( VertexInstanceID, 5, FVector2D( Data[ 1 ], Data[ 2 ] ) );
 
 		// leaf wind
 		DrawCall->GetProperty( SpeedTree::VERTEX_PROPERTY_WIND_EXTRA_DATA, VertexIndex, Data );
-		VertexInstanceUVs.GetArrayForIndex(6)[VertexInstanceID] = FVector2D( Data[ 0 ], Data[ 1 ] );
-		VertexInstanceUVs.GetArrayForIndex(7)[VertexInstanceID].X = Data[ 2 ];
+		VertexInstanceUVs.Set( VertexInstanceID, 6, FVector2D( Data[ 0 ], Data[ 1 ] ) );
+		VertexInstanceUVs.Set( VertexInstanceID, 7, FVector2D( Data[ 2 ], 0 ) );
 		DrawCall->GetProperty( SpeedTree::VERTEX_PROPERTY_WIND_FLAGS, VertexIndex, Data );
-		VertexInstanceUVs.GetArrayForIndex(7)[VertexInstanceID].Y = Data[ 0 ];
+		VertexInstanceUVs.Set( VertexInstanceID, 7, FVector2D( VertexInstanceUVs.Get( VertexInstanceID, 7 ).X, Data[ 0 ] ) );
 	}
 	return VertexInstanceID;
 }
@@ -1657,15 +1657,15 @@ UObject* USpeedTreeImportFactory::FactoryCreateBinary7(UClass* InClass, UObject*
 						FMeshDescription* MeshDescription = StaticMesh->CreateOriginalMeshDescription(LODIndex);
 						StaticMesh->RegisterMeshAttributes(*MeshDescription);
 						
-						TVertexAttributeArray<FVector>& VertexPositions = MeshDescription->VertexAttributes().GetAttributes<FVector>(MeshAttribute::Vertex::Position);
-						TEdgeAttributeArray<bool>& EdgeHardnesses = MeshDescription->EdgeAttributes().GetAttributes<bool>(MeshAttribute::Edge::IsHard);
-						TEdgeAttributeArray<float>& EdgeCreaseSharpnesses = MeshDescription->EdgeAttributes().GetAttributes<float>(MeshAttribute::Edge::CreaseSharpness);
-						TPolygonGroupAttributeArray<FName>& PolygonGroupImportedMaterialSlotNames = MeshDescription->PolygonGroupAttributes().GetAttributes<FName>(MeshAttribute::PolygonGroup::ImportedMaterialSlotName);
-						TVertexInstanceAttributeArray<FVector>& VertexInstanceNormals = MeshDescription->VertexInstanceAttributes().GetAttributes<FVector>(MeshAttribute::VertexInstance::Normal);
-						TVertexInstanceAttributeArray<FVector>& VertexInstanceTangents = MeshDescription->VertexInstanceAttributes().GetAttributes<FVector>(MeshAttribute::VertexInstance::Tangent);
-						TVertexInstanceAttributeArray<float>& VertexInstanceBinormalSigns = MeshDescription->VertexInstanceAttributes().GetAttributes<float>(MeshAttribute::VertexInstance::BinormalSign);
-						TVertexInstanceAttributeArray<FVector4>& VertexInstanceColors = MeshDescription->VertexInstanceAttributes().GetAttributes<FVector4>(MeshAttribute::VertexInstance::Color);
-						TVertexInstanceAttributeIndicesArray<FVector2D>& VertexInstanceUVs = MeshDescription->VertexInstanceAttributes().GetAttributesSet<FVector2D>(MeshAttribute::VertexInstance::TextureCoordinate);
+						TVertexAttributesRef<FVector> VertexPositions = MeshDescription->VertexAttributes().GetAttributesRef<FVector>(MeshAttribute::Vertex::Position);
+						TEdgeAttributesRef<bool> EdgeHardnesses = MeshDescription->EdgeAttributes().GetAttributesRef<bool>(MeshAttribute::Edge::IsHard);
+						TEdgeAttributesRef<float> EdgeCreaseSharpnesses = MeshDescription->EdgeAttributes().GetAttributesRef<float>(MeshAttribute::Edge::CreaseSharpness);
+						TPolygonGroupAttributesRef<FName> PolygonGroupImportedMaterialSlotNames = MeshDescription->PolygonGroupAttributes().GetAttributesRef<FName>(MeshAttribute::PolygonGroup::ImportedMaterialSlotName);
+						TVertexInstanceAttributesRef<FVector> VertexInstanceNormals = MeshDescription->VertexInstanceAttributes().GetAttributesRef<FVector>(MeshAttribute::VertexInstance::Normal);
+						TVertexInstanceAttributesRef<FVector> VertexInstanceTangents = MeshDescription->VertexInstanceAttributes().GetAttributesRef<FVector>(MeshAttribute::VertexInstance::Tangent);
+						TVertexInstanceAttributesRef<float> VertexInstanceBinormalSigns = MeshDescription->VertexInstanceAttributes().GetAttributesRef<float>(MeshAttribute::VertexInstance::BinormalSign);
+						TVertexInstanceAttributesRef<FVector4> VertexInstanceColors = MeshDescription->VertexInstanceAttributes().GetAttributesRef<FVector4>(MeshAttribute::VertexInstance::Color);
+						TVertexInstanceAttributesRef<FVector2D> VertexInstanceUVs = MeshDescription->VertexInstanceAttributes().GetAttributesRef<FVector2D>(MeshAttribute::VertexInstance::TextureCoordinate);
 
 						TMap<int32, FPolygonGroupID> MaterialToPolygonGroup;
 						for (int32 MatIndex = 0; MatIndex < StaticMesh->StaticMaterials.Num(); ++MatIndex)
@@ -1760,7 +1760,7 @@ UObject* USpeedTreeImportFactory::FactoryCreateBinary7(UClass* InClass, UObject*
 
 							FPolygonGroupID CurrentPolygonGroupID = MaterialToPolygonGroup[MaterialIndex];
 
-							int32 IndexOffset = VertexPositions.Num();
+							int32 IndexOffset = VertexPositions.GetNumElements();
 
 							for (int32 VertexIndex = 0; VertexIndex < DrawCall->m_nNumVertices; ++VertexIndex)
 							{
@@ -1859,15 +1859,15 @@ UObject* USpeedTreeImportFactory::FactoryCreateBinary7(UClass* InClass, UObject*
 					FMeshDescription* MeshDescription = StaticMesh->CreateOriginalMeshDescription(LODIndex);
 					StaticMesh->RegisterMeshAttributes(*MeshDescription);
 
-					TVertexAttributeArray<FVector>& VertexPositions = MeshDescription->VertexAttributes().GetAttributes<FVector>(MeshAttribute::Vertex::Position);
-					TEdgeAttributeArray<bool>& EdgeHardnesses = MeshDescription->EdgeAttributes().GetAttributes<bool>(MeshAttribute::Edge::IsHard);
-					TEdgeAttributeArray<float>& EdgeCreaseSharpnesses = MeshDescription->EdgeAttributes().GetAttributes<float>(MeshAttribute::Edge::CreaseSharpness);
-					TPolygonGroupAttributeArray<FName>& PolygonGroupImportedMaterialSlotNames = MeshDescription->PolygonGroupAttributes().GetAttributes<FName>(MeshAttribute::PolygonGroup::ImportedMaterialSlotName);
-					TVertexInstanceAttributeArray<FVector>& VertexInstanceNormals = MeshDescription->VertexInstanceAttributes().GetAttributes<FVector>(MeshAttribute::VertexInstance::Normal);
-					TVertexInstanceAttributeArray<FVector>& VertexInstanceTangents = MeshDescription->VertexInstanceAttributes().GetAttributes<FVector>(MeshAttribute::VertexInstance::Tangent);
-					TVertexInstanceAttributeArray<float>& VertexInstanceBinormalSigns = MeshDescription->VertexInstanceAttributes().GetAttributes<float>(MeshAttribute::VertexInstance::BinormalSign);
-					TVertexInstanceAttributeArray<FVector4>& VertexInstanceColors = MeshDescription->VertexInstanceAttributes().GetAttributes<FVector4>(MeshAttribute::VertexInstance::Color);
-					TVertexInstanceAttributeIndicesArray<FVector2D>& VertexInstanceUVs = MeshDescription->VertexInstanceAttributes().GetAttributesSet<FVector2D>(MeshAttribute::VertexInstance::TextureCoordinate);
+					TVertexAttributesRef<FVector> VertexPositions = MeshDescription->VertexAttributes().GetAttributesRef<FVector>(MeshAttribute::Vertex::Position);
+					TEdgeAttributesRef<bool> EdgeHardnesses = MeshDescription->EdgeAttributes().GetAttributesRef<bool>(MeshAttribute::Edge::IsHard);
+					TEdgeAttributesRef<float> EdgeCreaseSharpnesses = MeshDescription->EdgeAttributes().GetAttributesRef<float>(MeshAttribute::Edge::CreaseSharpness);
+					TPolygonGroupAttributesRef<FName> PolygonGroupImportedMaterialSlotNames = MeshDescription->PolygonGroupAttributes().GetAttributesRef<FName>(MeshAttribute::PolygonGroup::ImportedMaterialSlotName);
+					TVertexInstanceAttributesRef<FVector> VertexInstanceNormals = MeshDescription->VertexInstanceAttributes().GetAttributesRef<FVector>(MeshAttribute::VertexInstance::Normal);
+					TVertexInstanceAttributesRef<FVector> VertexInstanceTangents = MeshDescription->VertexInstanceAttributes().GetAttributesRef<FVector>(MeshAttribute::VertexInstance::Tangent);
+					TVertexInstanceAttributesRef<float> VertexInstanceBinormalSigns = MeshDescription->VertexInstanceAttributes().GetAttributesRef<float>(MeshAttribute::VertexInstance::BinormalSign);
+					TVertexInstanceAttributesRef<FVector4> VertexInstanceColors = MeshDescription->VertexInstanceAttributes().GetAttributesRef<FVector4>(MeshAttribute::VertexInstance::Color);
+					TVertexInstanceAttributesRef<FVector2D> VertexInstanceUVs = MeshDescription->VertexInstanceAttributes().GetAttributesRef<FVector2D>(MeshAttribute::VertexInstance::TextureCoordinate);
 
 					//Speedtree use UVs to store is data
 					VertexInstanceUVs.SetNumIndices(2);
@@ -1927,7 +1927,7 @@ UObject* USpeedTreeImportFactory::FactoryCreateBinary7(UClass* InClass, UObject*
 						const float* TexCoords = &SpeedTreeGeometry->m_sVertBBs.m_pTexCoords[BillboardIndex * 4];
 						bool bRotated = (SpeedTreeGeometry->m_sVertBBs.m_pRotated[BillboardIndex] == 1);
 
-						int32 IndexOffset = VertexPositions.Num();
+						int32 IndexOffset = VertexPositions.GetNumElements();
 					
 						// position
 						for (int32 VertexIndex = 0; VertexIndex < NumVertices; ++VertexIndex)
@@ -1958,15 +1958,15 @@ UObject* USpeedTreeImportFactory::FactoryCreateBinary7(UClass* InClass, UObject*
 								VertexInstanceBinormalSigns[VertexInstanceID] = GetBasisDeterminantSign(TangentX.GetSafeNormal(), TangentY.GetSafeNormal(), TangentZ.GetSafeNormal());
 								if (bRotated)
 								{
-									VertexInstanceUVs.GetArrayForIndex(0)[VertexInstanceID] = FVector2D(TexCoords[0] + Vertex[1] * TexCoords[2], TexCoords[1] + Vertex[0] * TexCoords[3]);
+									VertexInstanceUVs.Set(VertexInstanceID, 0, FVector2D(TexCoords[0] + Vertex[1] * TexCoords[2], TexCoords[1] + Vertex[0] * TexCoords[3]));
 								}
 								else
 								{
-									VertexInstanceUVs.GetArrayForIndex(0)[VertexInstanceID] = FVector2D(TexCoords[0] + Vertex[0] * TexCoords[2], TexCoords[1] + Vertex[1] * TexCoords[3]);
+									VertexInstanceUVs.Set(VertexInstanceID, 0, FVector2D(TexCoords[0] + Vertex[0] * TexCoords[2], TexCoords[1] + Vertex[1] * TexCoords[3]));
 								}
 
 								// lightmap coord
-								VertexInstanceUVs.GetArrayForIndex(1)[VertexInstanceID] = VertexInstanceUVs.GetArrayForIndex(0)[VertexInstanceID];
+								VertexInstanceUVs.Set(VertexInstanceID, 1, VertexInstanceUVs.Get(VertexInstanceID, 0));
 
 								CornerVertexInstanceIDs[Corner] = VertexInstanceID;
 								CornerVertexIDs[Corner] = VertexID;
@@ -2215,15 +2215,15 @@ UObject* USpeedTreeImportFactory::FactoryCreateBinary8(UClass* InClass, UObject*
 			FMeshDescription* MeshDescription = StaticMesh->CreateOriginalMeshDescription(LODIndex);
 			StaticMesh->RegisterMeshAttributes(*MeshDescription);
 
-			TVertexAttributeArray<FVector>& VertexPositions = MeshDescription->VertexAttributes().GetAttributes<FVector>(MeshAttribute::Vertex::Position);
-			TEdgeAttributeArray<bool>& EdgeHardnesses = MeshDescription->EdgeAttributes().GetAttributes<bool>(MeshAttribute::Edge::IsHard);
-			TEdgeAttributeArray<float>& EdgeCreaseSharpnesses = MeshDescription->EdgeAttributes().GetAttributes<float>(MeshAttribute::Edge::CreaseSharpness);
-			TPolygonGroupAttributeArray<FName>& PolygonGroupImportedMaterialSlotNames = MeshDescription->PolygonGroupAttributes().GetAttributes<FName>(MeshAttribute::PolygonGroup::ImportedMaterialSlotName);
-			TVertexInstanceAttributeArray<FVector>& VertexInstanceNormals = MeshDescription->VertexInstanceAttributes().GetAttributes<FVector>(MeshAttribute::VertexInstance::Normal);
-			TVertexInstanceAttributeArray<FVector>& VertexInstanceTangents = MeshDescription->VertexInstanceAttributes().GetAttributes<FVector>(MeshAttribute::VertexInstance::Tangent);
-			TVertexInstanceAttributeArray<float>& VertexInstanceBinormalSigns = MeshDescription->VertexInstanceAttributes().GetAttributes<float>(MeshAttribute::VertexInstance::BinormalSign);
-			TVertexInstanceAttributeArray<FVector4>& VertexInstanceColors = MeshDescription->VertexInstanceAttributes().GetAttributes<FVector4>(MeshAttribute::VertexInstance::Color);
-			TVertexInstanceAttributeIndicesArray<FVector2D>& VertexInstanceUVs = MeshDescription->VertexInstanceAttributes().GetAttributesSet<FVector2D>(MeshAttribute::VertexInstance::TextureCoordinate);
+			TVertexAttributesRef<FVector> VertexPositions = MeshDescription->VertexAttributes().GetAttributesRef<FVector>(MeshAttribute::Vertex::Position);
+			TEdgeAttributesRef<bool> EdgeHardnesses = MeshDescription->EdgeAttributes().GetAttributesRef<bool>(MeshAttribute::Edge::IsHard);
+			TEdgeAttributesRef<float> EdgeCreaseSharpnesses = MeshDescription->EdgeAttributes().GetAttributesRef<float>(MeshAttribute::Edge::CreaseSharpness);
+			TPolygonGroupAttributesRef<FName> PolygonGroupImportedMaterialSlotNames = MeshDescription->PolygonGroupAttributes().GetAttributesRef<FName>(MeshAttribute::PolygonGroup::ImportedMaterialSlotName);
+			TVertexInstanceAttributesRef<FVector> VertexInstanceNormals = MeshDescription->VertexInstanceAttributes().GetAttributesRef<FVector>(MeshAttribute::VertexInstance::Normal);
+			TVertexInstanceAttributesRef<FVector> VertexInstanceTangents = MeshDescription->VertexInstanceAttributes().GetAttributesRef<FVector>(MeshAttribute::VertexInstance::Tangent);
+			TVertexInstanceAttributesRef<float> VertexInstanceBinormalSigns = MeshDescription->VertexInstanceAttributes().GetAttributesRef<float>(MeshAttribute::VertexInstance::BinormalSign);
+			TVertexInstanceAttributesRef<FVector4> VertexInstanceColors = MeshDescription->VertexInstanceAttributes().GetAttributesRef<FVector4>(MeshAttribute::VertexInstance::Color);
+			TVertexInstanceAttributesRef<FVector2D> VertexInstanceUVs = MeshDescription->VertexInstanceAttributes().GetAttributesRef<FVector2D>(MeshAttribute::VertexInstance::TextureCoordinate);
 
 			//Speedtree use 8 UVs to store is data
 			VertexInstanceUVs.SetNumIndices(8);
@@ -2323,57 +2323,57 @@ UObject* USpeedTreeImportFactory::FactoryCreateBinary8(UClass* InClass, UObject*
 						// 7	0					0					Leaf Wind Z, Leaf Group
 
 						// diffuse
-						VertexInstanceUVs.GetArrayForIndex(0)[VertexInstanceID] = FVector2D(Vertex.m_vTexCoord.x, Vertex.m_vTexCoord.y);
+						VertexInstanceUVs.Set(VertexInstanceID, 0, FVector2D(Vertex.m_vTexCoord.x, Vertex.m_vTexCoord.y));
 
 						// lightmap
-						VertexInstanceUVs.GetArrayForIndex(1)[VertexInstanceID] = FVector2D(Vertex.m_vLightmapTexCoord.x, Vertex.m_vLightmapTexCoord.y);
+						VertexInstanceUVs.Set(VertexInstanceID, 1, FVector2D(Vertex.m_vLightmapTexCoord.x, Vertex.m_vLightmapTexCoord.y));
 
 						if (DrawCall.m_eWindGeometryType == SpeedTree8::Billboard)
 						{
-							VertexInstanceUVs.GetArrayForIndex(2)[VertexInstanceID] = FVector2D((Vertex.m_vNormal.z > 0.5f) ? 1.0f : 0.0f, 0.0f);
+							VertexInstanceUVs.Set(VertexInstanceID, 2, FVector2D((Vertex.m_vNormal.z > 0.5f) ? 1.0f : 0.0f, 0.0f));
 						}
 						else
 						{
 							// branch wind
-							VertexInstanceUVs.GetArrayForIndex(2)[VertexInstanceID] = FVector2D(Vertex.m_vWindBranch.x, Vertex.m_vWindBranch.y);
+							VertexInstanceUVs.Set(VertexInstanceID, 2, FVector2D(Vertex.m_vWindBranch.x, Vertex.m_vWindBranch.y));
 
 							// lod
 							FVector vLodPosition = FVector(Vertex.m_vAnchor.x, Vertex.m_vAnchor.y, Vertex.m_vAnchor.z) +
 								FVector(Vertex.m_vLodOffset.x, Vertex.m_vLodOffset.y, Vertex.m_vLodOffset.z);
-							VertexInstanceUVs.GetArrayForIndex(3)[VertexInstanceID] = FVector2D(vLodPosition[0], vLodPosition[1]);
-							VertexInstanceUVs.GetArrayForIndex(4)[VertexInstanceID] = FVector2D(vLodPosition[2], 0.0f);
+							VertexInstanceUVs.Set(VertexInstanceID, 3, FVector2D(vLodPosition[0], vLodPosition[1]));
+							VertexInstanceUVs.Set(VertexInstanceID, 4, FVector2D(vLodPosition[2], 0.0f));
 
 							// other
 							if (DrawCall.m_eWindGeometryType == SpeedTree8::Branch)
 							{
 								// detail (not used in v8)
-								VertexInstanceUVs.GetArrayForIndex(5)[VertexInstanceID] = FVector2D(0.0f, 0.0f);
+								VertexInstanceUVs.Set(VertexInstanceID, 5, FVector2D(0.0f, 0.0f));
 
 								// branch seam
-								VertexInstanceUVs.GetArrayForIndex(6)[VertexInstanceID] = FVector2D(0.0f, 0.0f);
-								VertexInstanceUVs.GetArrayForIndex(4)[VertexInstanceID].Y = Vertex.m_fBlendWeight;
+								VertexInstanceUVs.Set(VertexInstanceID, 6, FVector2D(0.0f, 0.0f));
+								VertexInstanceUVs.Set(VertexInstanceID, 4, FVector2D(VertexInstanceUVs.Get(VertexInstanceID, 4).X, Vertex.m_fBlendWeight));
 
 								// keep alignment
-								VertexInstanceUVs.GetArrayForIndex(7)[VertexInstanceID] = FVector2D(0.0f, 0.0f);
+								VertexInstanceUVs.Set(VertexInstanceID, 7, FVector2D(0.0f, 0.0f));
 							}
 							else if (DrawCall.m_eWindGeometryType == SpeedTree8::Frond)
 							{
 								// frond wind
-								VertexInstanceUVs.GetArrayForIndex(5)[VertexInstanceID] = FVector2D(Vertex.m_vWindNonBranch.x, Vertex.m_vWindNonBranch.y);
-								VertexInstanceUVs.GetArrayForIndex(6)[VertexInstanceID] = FVector2D(Vertex.m_vWindNonBranch.z, 0.0f);
+								VertexInstanceUVs.Set(VertexInstanceID, 5, FVector2D(Vertex.m_vWindNonBranch.x, Vertex.m_vWindNonBranch.y));
+								VertexInstanceUVs.Set(VertexInstanceID, 6, FVector2D(Vertex.m_vWindNonBranch.z, 0.0f));
 
 								// keep alignment
-								VertexInstanceUVs.GetArrayForIndex(7)[VertexInstanceID] = FVector2D(0.0f, 0.0f);
+								VertexInstanceUVs.Set(VertexInstanceID, 7, FVector2D(0.0f, 0.0f));
 							}
 							else if (DrawCall.m_eWindGeometryType == SpeedTree8::Leaf || DrawCall.m_eWindGeometryType == SpeedTree8::FacingLeaf)
 							{
 								// anchor
-								VertexInstanceUVs.GetArrayForIndex(4)[VertexInstanceID].Y = Vertex.m_vAnchor.x;
-								VertexInstanceUVs.GetArrayForIndex(5)[VertexInstanceID] = FVector2D(Vertex.m_vAnchor.y, Vertex.m_vAnchor.z);
+								VertexInstanceUVs.Set(VertexInstanceID, 4, FVector2D(VertexInstanceUVs.Get(VertexInstanceID, 4).X, Vertex.m_vAnchor.x));
+								VertexInstanceUVs.Set(VertexInstanceID, 5, FVector2D(Vertex.m_vAnchor.y, Vertex.m_vAnchor.z));
 
 								// leaf wind
-								VertexInstanceUVs.GetArrayForIndex(6)[VertexInstanceID] = FVector2D(Vertex.m_vWindNonBranch.x, Vertex.m_vWindNonBranch.y);
-								VertexInstanceUVs.GetArrayForIndex(7)[VertexInstanceID] = FVector2D(Vertex.m_vWindNonBranch.z, (Vertex.m_bWindLeaf2Flag ? 1.0f : 0.0f));
+								VertexInstanceUVs.Set(VertexInstanceID, 6, FVector2D(Vertex.m_vWindNonBranch.x, Vertex.m_vWindNonBranch.y));
+								VertexInstanceUVs.Set(VertexInstanceID, 7, FVector2D(Vertex.m_vWindNonBranch.z, (Vertex.m_bWindLeaf2Flag ? 1.0f : 0.0f)));
 							}
 						}
 						CornerVertexInstanceIDs[Corner] = VertexInstanceID;
diff --git a/Engine/Plugins/Experimental/AlembicImporter/Source/AlembicLibrary/Private/AbcImporter.cpp b/Engine/Plugins/Experimental/AlembicImporter/Source/AlembicLibrary/Private/AbcImporter.cpp
--- a/Engine/Plugins/Experimental/AlembicImporter/Source/AlembicLibrary/Private/AbcImporter.cpp
+++ b/Engine/Plugins/Experimental/AlembicImporter/Source/AlembicLibrary/Private/AbcImporter.cpp
@@ -1256,15 +1256,15 @@ const uint32 FAbcImporter::GetNumMeshTracks() const
 
 void FAbcImporter::GenerateMeshDescriptionFromSample(const FAbcMeshSample* Sample, FMeshDescription* MeshDescription, UStaticMesh* StaticMesh)
 {
-	TVertexAttributeArray<FVector>& VertexPositions = MeshDescription->VertexAttributes().GetAttributes<FVector>(MeshAttribute::Vertex::Position);
-	TEdgeAttributeArray<bool>& EdgeHardnesses = MeshDescription->EdgeAttributes().GetAttributes<bool>(MeshAttribute::Edge::IsHard);
-	TEdgeAttributeArray<float>& EdgeCreaseSharpnesses = MeshDescription->EdgeAttributes().GetAttributes<float>(MeshAttribute::Edge::CreaseSharpness);
-	TPolygonGroupAttributeArray<FName>& PolygonGroupImportedMaterialSlotNames = MeshDescription->PolygonGroupAttributes().GetAttributes<FName>(MeshAttribute::PolygonGroup::ImportedMaterialSlotName);
-	TVertexInstanceAttributeArray<FVector>& VertexInstanceNormals = MeshDescription->VertexInstanceAttributes().GetAttributes<FVector>(MeshAttribute::VertexInstance::Normal);
-	TVertexInstanceAttributeArray<FVector>& VertexInstanceTangents = MeshDescription->VertexInstanceAttributes().GetAttributes<FVector>(MeshAttribute::VertexInstance::Tangent);
-	TVertexInstanceAttributeArray<float>& VertexInstanceBinormalSigns = MeshDescription->VertexInstanceAttributes().GetAttributes<float>(MeshAttribute::VertexInstance::BinormalSign);
-	TVertexInstanceAttributeArray<FVector4>& VertexInstanceColors = MeshDescription->VertexInstanceAttributes().GetAttributes<FVector4>(MeshAttribute::VertexInstance::Color);
-	TVertexInstanceAttributeIndicesArray<FVector2D>& VertexInstanceUVs = MeshDescription->VertexInstanceAttributes().GetAttributesSet<FVector2D>(MeshAttribute::VertexInstance::TextureCoordinate);
+	TVertexAttributesRef<FVector> VertexPositions = MeshDescription->VertexAttributes().GetAttributesRef<FVector>(MeshAttribute::Vertex::Position);
+	TEdgeAttributesRef<bool> EdgeHardnesses = MeshDescription->EdgeAttributes().GetAttributesRef<bool>(MeshAttribute::Edge::IsHard);
+	TEdgeAttributesRef<float> EdgeCreaseSharpnesses = MeshDescription->EdgeAttributes().GetAttributesRef<float>(MeshAttribute::Edge::CreaseSharpness);
+	TPolygonGroupAttributesRef<FName> PolygonGroupImportedMaterialSlotNames = MeshDescription->PolygonGroupAttributes().GetAttributesRef<FName>(MeshAttribute::PolygonGroup::ImportedMaterialSlotName);
+	TVertexInstanceAttributesRef<FVector> VertexInstanceNormals = MeshDescription->VertexInstanceAttributes().GetAttributesRef<FVector>(MeshAttribute::VertexInstance::Normal);
+	TVertexInstanceAttributesRef<FVector> VertexInstanceTangents = MeshDescription->VertexInstanceAttributes().GetAttributesRef<FVector>(MeshAttribute::VertexInstance::Tangent);
+	TVertexInstanceAttributesRef<float> VertexInstanceBinormalSigns = MeshDescription->VertexInstanceAttributes().GetAttributesRef<float>(MeshAttribute::VertexInstance::BinormalSign);
+	TVertexInstanceAttributesRef<FVector4> VertexInstanceColors = MeshDescription->VertexInstanceAttributes().GetAttributesRef<FVector4>(MeshAttribute::VertexInstance::Color);
+	TVertexInstanceAttributesRef<FVector2D> VertexInstanceUVs = MeshDescription->VertexInstanceAttributes().GetAttributesRef<FVector2D>(MeshAttribute::VertexInstance::TextureCoordinate);
 
 	//Speedtree use UVs to store is data
 	VertexInstanceUVs.SetNumIndices(Sample->NumUVSets);
@@ -1316,7 +1316,7 @@ void FAbcImporter::GenerateMeshDescriptionFromSample(const FAbcMeshSample* Sampl
 
 			for (uint32 UVIndex = 0; UVIndex < Sample->NumUVSets; ++UVIndex)
 			{
-				VertexInstanceUVs.GetArrayForIndex(UVIndex)[VertexInstanceID] = Sample->UVs[UVIndex][IndiceIndex];
+				VertexInstanceUVs.Set(VertexInstanceID, UVIndex, Sample->UVs[UVIndex][IndiceIndex]);
 			}
 			CornerVertexInstanceIDs[Corner] = VertexInstanceID;
 			CornerVertexIDs[Corner] = VertexID;
diff --git a/Engine/Plugins/Runtime/EditableMesh/Source/EditableMesh/EditableMesh.cpp b/Engine/Plugins/Runtime/EditableMesh/Source/EditableMesh/EditableMesh.cpp
--- a/Engine/Plugins/Runtime/EditableMesh/Source/EditableMesh/EditableMesh.cpp
+++ b/Engine/Plugins/Runtime/EditableMesh/Source/EditableMesh/EditableMesh.cpp
@@ -91,16 +91,15 @@ template <> struct TIsAttributeTypeLerpable<int> { static const bool Value = tru
 
 
 template <typename T, typename ElementIDType>
-static typename TEnableIf<TIsAttributeTypeLerpable<T>::Value>::Type InterpAttributeOfType( FMeshElementAttributeList& AttributeList, const FName AttributeName, const TAttributeIndicesArray<T, ElementIDType>& AttributeIndicesArray, const ElementIDType ID0, const ElementIDType ID1, const float Alpha )
+static typename TEnableIf<TIsAttributeTypeLerpable<T>::Value>::Type InterpAttributeOfType( FMeshElementAttributeList& AttributeList, const FName AttributeName, TMeshAttributesConstRef<ElementIDType, T> AttributeArrayRef, const ElementIDType ID0, const ElementIDType ID1, const float Alpha )
 {
-	for( int32 Index = 0; Index < AttributeIndicesArray.GetNumIndices(); ++Index )
+	for( int32 Index = 0; Index < AttributeArrayRef.GetNumIndices(); ++Index )
 	{
-		const TMeshAttributeArray<T, ElementIDType>& Array = AttributeIndicesArray.GetArrayForIndex( Index );
-		const T& Value0 = Array[ ID0 ];
-		const T& Value1 = Array[ ID1 ];
+		const T& Value0 = AttributeArrayRef.Get( ID0, Index );
+		const T& Value1 = AttributeArrayRef.Get( ID1, Index );
 
 		AttributeList.Attributes.Emplace( AttributeName, Index, 
-			FMeshElementAttributeValue( EnumHasAllFlags( AttributeIndicesArray.GetFlags(), EMeshAttributeFlags::Lerpable ) ?
+			FMeshElementAttributeValue( EnumHasAllFlags( AttributeArrayRef.GetFlags(), EMeshAttributeFlags::Lerpable ) ?
 				FMath::Lerp( Value0, Value1, Alpha ) :
 				( Alpha <= 0.5f ) ? Value0 : Value1 )
 			);
@@ -108,14 +107,13 @@ static typename TEnableIf<TIsAttributeTypeLerpable<T>::Value>::Type InterpAttrib
 }
 
 template <typename T, typename ElementIDType>
-static typename TEnableIf<!TIsAttributeTypeLerpable<T>::Value>::Type InterpAttributeOfType( FMeshElementAttributeList& AttributeList, const FName AttributeName, const TAttributeIndicesArray<T, ElementIDType>& AttributeIndicesArray, const ElementIDType ID0, const ElementIDType ID1, const float Alpha )
+static typename TEnableIf<!TIsAttributeTypeLerpable<T>::Value>::Type InterpAttributeOfType( FMeshElementAttributeList& AttributeList, const FName AttributeName, TMeshAttributesConstRef<ElementIDType, T> AttributeArrayRef, const ElementIDType ID0, const ElementIDType ID1, const float Alpha )
 {
 	// Non-Lerpable attribute type: just choose whichever value alpha is closest to
-	for( int32 Index = 0; Index < AttributeIndicesArray.GetNumIndices(); ++Index )
+	for( int32 Index = 0; Index < AttributeArrayRef.GetNumIndices(); ++Index )
 	{
-		const TMeshAttributeArray<T, ElementIDType>& Array = AttributeIndicesArray.GetArrayForIndex( Index );
-		const T& Value0 = Array[ ID0 ];
-		const T& Value1 = Array[ ID1 ];
+		const T& Value0 = AttributeArrayRef.Get( ID0, Index );
+		const T& Value1 = AttributeArrayRef.Get( ID1, Index );
 
 		AttributeList.Attributes.Emplace( AttributeName, Index, FMeshElementAttributeValue( ( Alpha <= 0.5f ) ? Value0 : Value1 ) );
 	}
@@ -129,10 +127,10 @@ static typename TEnableIf<!TIsAttributeTypeLerpable<T>::Value>::Type InterpAttri
 template <typename ElementIDType>
 static void InterpAllAttributes( FMeshElementAttributeList& AttributeList, const TAttributesSet<ElementIDType>& AttributesSet, const ElementIDType ID0, const ElementIDType ID1, const float Alpha )
 {
-	AttributesSet.ForEachAttributeIndicesArray(
-		[ &AttributeList, ID0, ID1, Alpha ]( const FName AttributeName, const auto& AttributeIndicesArray )
+	AttributesSet.ForEach(
+		[ &AttributeList, ID0, ID1, Alpha ]( const FName AttributeName, auto AttributeArrayRef )
 		{
-			InterpAttributeOfType( AttributeList, AttributeName, AttributeIndicesArray, ID0, ID1, Alpha );
+			InterpAttributeOfType( AttributeList, AttributeName, AttributeArrayRef, ID0, ID1, Alpha );
 		}
 	);
 }
@@ -145,12 +143,12 @@ static void InterpAllAttributes( FMeshElementAttributeList& AttributeList, const
 template <typename ElementIDType, typename Predicate>
 static void InterpAttributesByPredicate( FMeshElementAttributeList& AttributeList, const TAttributesSet<ElementIDType>& AttributesSet, const ElementIDType ID0, const ElementIDType ID1, const float Alpha, const Predicate& Pred )
 {
-	AttributesSet.ForEachAttributeIndicesArray(
-		[ &AttributeList, ID0, ID1, Alpha, &Pred ]( const FName AttributeName, const auto& AttributeIndicesArray )
+	AttributesSet.ForEach(
+		[ &AttributeList, ID0, ID1, Alpha, &Pred ]( const FName AttributeName, auto AttributeArrayRef )
 		{
-			if( Pred( AttributeName, AttributeIndicesArray ) )
+			if( Pred( AttributeName, AttributeArrayRef ) )
 			{
-				InterpAttributeOfType( AttributeList, AttributeName, AttributeIndicesArray, ID0, ID1, Alpha );
+				InterpAttributeOfType( AttributeList, AttributeName, AttributeArrayRef, ID0, ID1, Alpha );
 			}
 		}
 	);
@@ -158,17 +156,16 @@ static void InterpAttributesByPredicate( FMeshElementAttributeList& AttributeLis
 
 
 template <typename T, typename ElementIDType>
-static typename TEnableIf<TIsAttributeTypeLerpable<T>::Value>::Type InterpAttributeOfType( FMeshElementAttributeList& AttributeList, const FName AttributeName, const TAttributeIndicesArray<T, ElementIDType>& AttributeIndicesArray, const ElementIDType ID0, const ElementIDType ID1, const ElementIDType ID2, const FVector Weights )
+static typename TEnableIf<TIsAttributeTypeLerpable<T>::Value>::Type InterpAttributeOfType( FMeshElementAttributeList& AttributeList, const FName AttributeName, TMeshAttributesConstRef<ElementIDType, T> AttributeArrayRef, const ElementIDType ID0, const ElementIDType ID1, const ElementIDType ID2, const FVector Weights )
 {
-	if( EnumHasAllFlags( AttributeIndicesArray.GetFlags(), EMeshAttributeFlags::Lerpable ) )
+	if( EnumHasAllFlags( AttributeArrayRef.GetFlags(), EMeshAttributeFlags::Lerpable ) )
 	{
 		// Get result by interpolating between the three values according to the barycentric weight
-		for( int32 Index = 0; Index < AttributeIndicesArray.GetNumIndices(); ++Index )
+		for( int32 Index = 0; Index < AttributeArrayRef.GetNumIndices(); ++Index )
 		{
-			const TMeshAttributeArray<T, ElementIDType>& Array = AttributeIndicesArray.GetArrayForIndex( Index );
-			const T& Value0 = Array[ ID0 ];
-			const T& Value1 = Array[ ID1 ];
-			const T& Value2 = Array[ ID2 ];
+			const T& Value0 = AttributeArrayRef.Get( ID0, Index );
+			const T& Value1 = AttributeArrayRef.Get( ID1, Index );
+			const T& Value2 = AttributeArrayRef.Get( ID2, Index );
 
 			AttributeList.Attributes.Emplace( AttributeName, Index, FMeshElementAttributeValue( Value0 * Weights.X + Value1 * Weights.Y + Value2 * Weights.Z ) );
 		}
@@ -178,24 +175,22 @@ static typename TEnableIf<TIsAttributeTypeLerpable<T>::Value>::Type InterpAttrib
 		// Non-Lerpable attribute type: just choose whichever value the barycentric weight is closest to
 		const ElementIDType ClosestID = ( Weights.X > Weights.Y ) ? ( ( Weights.X > Weights.Z ) ? ID0 : ID2 ) : ( ( Weights.Y > Weights.Z ) ? ID1 : ID2 );
 
-		for( int32 Index = 0; Index < AttributeIndicesArray.GetNumIndices(); ++Index )
+		for( int32 Index = 0; Index < AttributeArrayRef.GetNumIndices(); ++Index )
 		{
-			const TMeshAttributeArray<T, ElementIDType>& Array = AttributeIndicesArray.GetArrayForIndex( Index );
-			AttributeList.Attributes.Emplace( AttributeName, Index, FMeshElementAttributeValue( Array[ ClosestID ] ) );
+			AttributeList.Attributes.Emplace( AttributeName, Index, FMeshElementAttributeValue( AttributeArrayRef.Get( ClosestID, Index ) ) );
 		}
 	}
 }
 
 template <typename T, typename ElementIDType>
-static typename TEnableIf<!TIsAttributeTypeLerpable<T>::Value>::Type InterpAttributeOfType( FMeshElementAttributeList& AttributeList, const FName AttributeName, const TAttributeIndicesArray<T, ElementIDType>& AttributeIndicesArray, const ElementIDType ID0, const ElementIDType ID1, const ElementIDType ID2, const FVector Weights )
+static typename TEnableIf<!TIsAttributeTypeLerpable<T>::Value>::Type InterpAttributeOfType( FMeshElementAttributeList& AttributeList, const FName AttributeName, TMeshAttributesConstRef<ElementIDType, T> AttributeArrayRef, const ElementIDType ID0, const ElementIDType ID1, const ElementIDType ID2, const FVector Weights )
 {
 	// Non-Lerpable attribute type: just choose whichever value the barycentric weight is closest to
 	const ElementIDType ClosestID = ( Weights.X > Weights.Y ) ? ( ( Weights.X > Weights.Z ) ? ID0 : ID2 ) : ( ( Weights.Y > Weights.Z ) ? ID1 : ID2 );
 
-	for( int32 Index = 0; Index < AttributeIndicesArray.GetNumIndices(); ++Index )
+	for( int32 Index = 0; Index < AttributeArrayRef.GetNumIndices(); ++Index )
 	{
-		const TMeshAttributeArray<T, ElementIDType>& Array = AttributeIndicesArray.GetArrayForIndex( Index );
-		AttributeList.Attributes.Emplace( AttributeName, Index, FMeshElementAttributeValue( Array[ ClosestID ] ) );
+		AttributeList.Attributes.Emplace( AttributeName, Index, FMeshElementAttributeValue( AttributeArrayRef.Get( ClosestID, Index ) ) );
 	}
 }
 
@@ -207,10 +202,10 @@ static typename TEnableIf<!TIsAttributeTypeLerpable<T>::Value>::Type InterpAttri
 template <typename ElementIDType>
 static void InterpAllAttributes( FMeshElementAttributeList& AttributeList, const TAttributesSet<ElementIDType>& AttributesSet, const ElementIDType ID0, const ElementIDType ID1, const ElementIDType ID2, const FVector Weights )
 {
-	AttributesSet.ForEachAttributeIndicesArray(
-		[ &AttributeList, ID0, ID1, ID2, Weights ]( const FName AttributeName, const auto& AttributeIndicesArray )
+	AttributesSet.ForEach(
+		[ &AttributeList, ID0, ID1, ID2, Weights ]( const FName AttributeName, auto AttributeArrayRef )
 		{
-			InterpAttributeOfType( AttributeList, AttributeName, AttributeIndicesArray, ID0, ID1, ID2, Weights );
+			InterpAttributeOfType( AttributeList, AttributeName, AttributeArrayRef, ID0, ID1, ID2, Weights );
 		}
 	);
 }
@@ -223,12 +218,12 @@ static void InterpAllAttributes( FMeshElementAttributeList& AttributeList, const
 template <typename ElementIDType, typename Predicate>
 static void InterpAttributesByPredicate( FMeshElementAttributeList& AttributeList, const TAttributesSet<ElementIDType>& AttributesSet, const ElementIDType ID0, const ElementIDType ID1, const ElementIDType ID2, const FVector Weights, const Predicate& Pred )
 {
-	AttributesSet.ForEachAttributeIndicesArray(
-		[ &AttributeList, ID0, ID1, ID2, Weights, &Pred ]( const FName AttributeName, const auto& AttributeIndicesArray )
+	AttributesSet.ForEach(
+		[ &AttributeList, ID0, ID1, ID2, Weights, &Pred ]( const FName AttributeName, auto AttributeArrayRef )
 		{
-			if( Pred( AttributeName, AttributeIndicesArray ) )
+			if( Pred( AttributeName, AttributeArrayRef ) )
 			{
-				InterpAttributeOfType( AttributeList, AttributeName, AttributeIndicesArray, ID0, ID1, ID2, Weights );
+				InterpAttributeOfType( AttributeList, AttributeName, AttributeArrayRef, ID0, ID1, ID2, Weights );
 			}
 		}
 	);
@@ -248,15 +243,14 @@ template <typename ElementIDType, typename Predicate>
 static bool CompareAttributesByPredicate( const TAttributesSet<ElementIDType>& AttributesSet, const ElementIDType ID0, const ElementIDType ID1, const Predicate& Pred )
 {
 	bool bResult = true;
-	AttributesSet.ForEachAttributeIndicesArray(
-		[ ID0, ID1, &Pred, &bResult ]( const FName AttributeName, const auto& AttributeIndicesArray )
+	AttributesSet.ForEach(
+		[ ID0, ID1, &Pred, &bResult ]( const FName AttributeName, auto AttributeArrayRef )
 		{
-			if( Pred( AttributeName, AttributeIndicesArray ) )
+			if( Pred( AttributeName, AttributeArrayRef ) )
 			{
-				for( int32 Index = 0; Index < AttributeIndicesArray.GetNumIndices(); ++Index )
+				for( int32 Index = 0; Index < AttributeArrayRef.GetNumIndices(); ++Index )
 				{
-					auto& Array = AttributeIndicesArray.GetArrayForIndex( Index );
-					if( !AreAttributeValuesNearlyEqual( Array[ ID0 ], Array[ ID1 ] ) )
+					if( !AreAttributeValuesNearlyEqual( AttributeArrayRef.Get( ID0, Index ), AttributeArrayRef.Get( ID1, Index ) ) )
 					{
 						bResult = false;
 						return;
@@ -274,13 +268,12 @@ static bool CompareAttributesByPredicate( const TAttributesSet<ElementIDType>& A
 template <typename ElementIDType>
 static void BackupAllAttributes( FMeshElementAttributeList& AttributeList, const TAttributesSet<ElementIDType>& AttributesSet, const ElementIDType ElementID )
 {
-	AttributesSet.ForEachAttributeIndicesArray(
-		[ &AttributeList, ElementID ]( const FName AttributeName, const auto& AttributeIndicesArray )
+	AttributesSet.ForEach(
+		[ &AttributeList, ElementID ]( const FName AttributeName, auto AttributeArrayRef )
 		{
-			for( int32 Index = 0; Index < AttributeIndicesArray.GetNumIndices(); ++Index )
+			for( int32 Index = 0; Index < AttributeArrayRef.GetNumIndices(); ++Index )
 			{
-				const auto& Value = AttributeIndicesArray.GetArrayForIndex( Index )[ ElementID ];
-
+				const auto& Value = AttributeArrayRef.Get( ElementID, Index );
 				AttributeList.Attributes.Emplace( AttributeName, Index, FMeshElementAttributeValue( Value ) );
 			}
 		}
@@ -292,15 +285,14 @@ static void BackupAllAttributes( FMeshElementAttributeList& AttributeList, const
 template <typename ElementIDType, typename Predicate>
 static void BackupAttributesByPredicate( FMeshElementAttributeList& AttributeList, const TAttributesSet<ElementIDType>& AttributesSet, const ElementIDType ElementID, const Predicate& Pred )
 {
-	AttributesSet.ForEachAttributeIndicesArray(
-		[ &AttributeList, ElementID, &Pred ]( const FName AttributeName, const auto& AttributeIndicesArray )
+	AttributesSet.ForEach(
+		[ &AttributeList, ElementID, &Pred ]( const FName AttributeName, auto AttributeArrayRef )
 		{
-			if( Pred( AttributeName, AttributeIndicesArray ) )
+			if( Pred( AttributeName, AttributeArrayRef ) )
 			{
-				for( int32 Index = 0; Index < AttributeIndicesArray.GetNumIndices(); ++Index )
+				for( int32 Index = 0; Index < AttributeArrayRef.GetNumIndices(); ++Index )
 				{
-					const auto& Value = AttributeIndicesArray.GetArrayForIndex( Index )[ ElementID ];
-
+					const auto& Value = AttributeArrayRef.Get( ElementID, Index );
 					AttributeList.Attributes.Emplace( AttributeName, Index, FMeshElementAttributeValue( Value ) );
 				}
 			}
@@ -316,8 +308,7 @@ static void ApplyAttribute( TAttributesSet<ElementIDType>& AttributesSet, const
 	AttributeData.AttributeValue.Visit(
 		[ &AttributesSet, &AttributeData, ElementID ]( const auto& Value )
 		{
-			using AttributeType = typename TDecay<decltype( Value )>::Type;
-			AttributesSet.template SetAttribute<AttributeType>( ElementID, AttributeData.AttributeName, AttributeData.AttributeIndex, Value );
+			AttributesSet.SetAttribute( ElementID, AttributeData.AttributeName, AttributeData.AttributeIndex, Value );
 		}
 	);
 }
@@ -404,16 +395,13 @@ void UEditableMesh::SetMeshDescription( FMeshDescription* InMeshDescription )
 }
 
 
-template <typename ElementIDType>
-static void InvertRemapTable( TSparseArray<ElementIDType>& InvertedRemapTable, const TSparseArray<ElementIDType>& RemapTable )
+static void InvertRemapTable( TSparseArray<int32>& InvertedRemapTable, const TSparseArray<int32>& RemapTable )
 {
-	static_assert( TIsDerivedFrom<ElementIDType, FElementID>::IsDerived, "Remap array type must be derived from FElementID" );
-
 	InvertedRemapTable.Empty( RemapTable.Num() );
 
-	for( typename TSparseArray<ElementIDType>::TConstIterator It( RemapTable ); It; ++It )
+	for( typename TSparseArray<int32>::TConstIterator It( RemapTable ); It; ++It )
 	{
-		InvertedRemapTable.Insert( It->GetValue(), ElementIDType( It.GetIndex() ) );
+		InvertedRemapTable.Insert( *It, It.GetIndex() );
 	}
 }
 
@@ -1399,7 +1387,7 @@ void UEditableMesh::MergeVertexInstances()
 					GetMeshDescription()->VertexInstanceAttributes(),
 					VertexInstanceIDA,
 					VertexInstanceIDB,
-					[]( const FName AttributeName, const auto& AttributeIndicesArray ) { return EnumHasAllFlags( AttributeIndicesArray.GetFlags(), EMeshAttributeFlags::Mergeable ); }
+					[]( const FName AttributeName, const auto AttributeArrayRef ) { return EnumHasAllFlags( AttributeArrayRef.GetFlags(), EMeshAttributeFlags::Mergeable ); }
 				  ) )
 				{
 					// Change occurrences of VertexInstanceB for VertexInstanceA in VertexInstanceB's connected polygons.
@@ -1743,7 +1731,7 @@ FBox UEditableMesh::ComputeBoundingBox() const
 	FBox BoundingBox;
 	BoundingBox.Init();
 
-	const TVertexAttributeArray<FVector>& VertexPositions = GetMeshDescription()->VertexAttributes().GetAttributes<FVector>( MeshAttribute::Vertex::Position );
+	TVertexAttributesConstRef<FVector> VertexPositions = GetMeshDescription()->VertexAttributes().GetAttributesRef<FVector>( MeshAttribute::Vertex::Position );
 
 	for( const FVertexID VertexID : GetMeshDescription()->Vertices().GetElementIDs() )
 	{
@@ -1764,7 +1752,7 @@ FBoxSphereBounds UEditableMesh::ComputeBoundingBoxAndSphere() const
 	// Calculate the bounding sphere, using the center of the bounding box as the origin.
 	BoundingBoxAndSphere.SphereRadius = 0.0f;
 
-	const TVertexAttributeArray<FVector>& VertexPositions = GetMeshDescription()->VertexAttributes().GetAttributes<FVector>( MeshAttribute::Vertex::Position );
+	TVertexAttributesConstRef<FVector> VertexPositions = GetMeshDescription()->VertexAttributes().GetAttributesRef<FVector>( MeshAttribute::Vertex::Position );
 
 	for( const FVertexID VertexID : GetMeshDescription()->Vertices().GetElementIDs() )
 	{
@@ -1782,7 +1770,7 @@ FVector UEditableMesh::ComputePolygonCenter( const FPolygonID PolygonID ) const
 	static TArray<FVertexID> PerimeterVertexIDs;
 	GetPolygonPerimeterVertices( PolygonID, /* Out */ PerimeterVertexIDs );
 
-	const TVertexAttributeArray<FVector>& VertexPositions = GetMeshDescription()->VertexAttributes().GetAttributes<FVector>( MeshAttribute::Vertex::Position );
+	TVertexAttributesConstRef<FVector> VertexPositions = GetMeshDescription()->VertexAttributes().GetAttributesRef<FVector>( MeshAttribute::Vertex::Position );
 
 	for( const FVertexID VertexID : PerimeterVertexIDs )
 	{
@@ -1809,7 +1797,7 @@ FPlane UEditableMesh::ComputePolygonPlane( const FPolygonID PolygonID ) const
 	static TArray<FVertexID> PerimeterVertexIDs;
 	GetPolygonPerimeterVertices( PolygonID, /* Out */ PerimeterVertexIDs );
 
-	const TVertexAttributeArray<FVector>& VertexPositions = GetMeshDescription()->VertexAttributes().GetAttributes<FVector>( MeshAttribute::Vertex::Position );
+	TVertexAttributesConstRef<FVector> VertexPositions = GetMeshDescription()->VertexAttributes().GetAttributesRef<FVector>( MeshAttribute::Vertex::Position );
 
 	// Use 'Newell's Method' to compute a robust 'best fit' plane from the vertices of this polygon
 	for( int32 VertexNumberI = PerimeterVertexIDs.Num() - 1, VertexNumberJ = 0; VertexNumberJ < PerimeterVertexIDs.Num(); VertexNumberI = VertexNumberJ, VertexNumberJ++ )
@@ -1870,7 +1858,7 @@ void UEditableMesh::RefreshOpenSubdiv()
 					OsdCornerVertexIndices.Reset();
 					OsdCornerWeights.Reset();
 
-					const TVertexAttributeArray<float>& VertexCornerSharpnesses = GetMeshDescription()->VertexAttributes().GetAttributes<float>( MeshAttribute::Vertex::CornerSharpness );
+					const TVertexAttributesRef<float> VertexCornerSharpnesses = GetMeshDescription()->VertexAttributes().GetAttributesRef<float>( MeshAttribute::Vertex::CornerSharpness );
 
 					for( const FVertexID VertexID : GetMeshDescription()->Vertices().GetElementIDs() )
 					{
@@ -1889,7 +1877,7 @@ void UEditableMesh::RefreshOpenSubdiv()
 					OsdCreaseVertexIndexPairs.Reset();
 					OsdCreaseWeights.Reset();
 
-					const TEdgeAttributeArray<float>& EdgeCreaseSharpnesses = GetMeshDescription()->EdgeAttributes().GetAttributes<float>( MeshAttribute::Edge::CreaseSharpness );
+					const TEdgeAttributesRef<float> EdgeCreaseSharpnesses = GetMeshDescription()->EdgeAttributes().GetAttributesRef<float>( MeshAttribute::Edge::CreaseSharpness );
 
 					for( const FEdgeID EdgeID : GetMeshDescription()->Edges().GetElementIDs() )
 					{
@@ -2085,7 +2073,7 @@ void UEditableMesh::GenerateOpenSubdivLimitSurfaceData()
 
 				// @todo mesheditor perf: note that now we expose vertex positions as an array of FVector, so this is just copying from one array to another.
 				// This needs to be refactored to not require this.
-				const TVertexAttributeArray<FVector>& SrcVertexPositions = GetMeshDescription()->VertexAttributes().GetAttributes<FVector>( MeshAttribute::Vertex::Position );
+				const TVertexAttributesRef<FVector> SrcVertexPositions = GetMeshDescription()->VertexAttributes().GetAttributesRef<FVector>( MeshAttribute::Vertex::Position );
 
 				for( int32 VertexNumber = 0; VertexNumber < VertexArraySize; ++VertexNumber )
 				{
@@ -2109,8 +2097,8 @@ void UEditableMesh::GenerateOpenSubdivLimitSurfaceData()
 				FVarVertexDatas.Reset();
 				FVarVertexDatas.Reserve( OsdFVarIndicesPerFace.Num() );
 
-				const TAttributeIndicesArray<FVector2D, FVertexInstanceID>& VertexUVs = GetMeshDescription()->VertexInstanceAttributes().GetAttributesSet<FVector2D>( MeshAttribute::VertexInstance::TextureCoordinate );
-				const TVertexInstanceAttributeArray<FVector4>& VertexColors = GetMeshDescription()->VertexInstanceAttributes().GetAttributes<FVector4>( MeshAttribute::VertexInstance::Color );
+				const TVertexInstanceAttributesRef<FVector2D> VertexUVs = GetMeshDescription()->VertexInstanceAttributes().GetAttributesRef<FVector2D>( MeshAttribute::VertexInstance::TextureCoordinate );
+				const TVertexInstanceAttributesRef<FVector4> VertexColors = GetMeshDescription()->VertexInstanceAttributes().GetAttributesRef<FVector4>( MeshAttribute::VertexInstance::Color );
 
 				int32 NumPolygonsSoFar = 0;
 				for( const FPolygonGroupID PolygonGroupID : GetMeshDescription()->PolygonGroups().GetElementIDs() )
@@ -2125,8 +2113,8 @@ void UEditableMesh::GenerateOpenSubdivLimitSurfaceData()
 						for( const FVertexInstanceID VertexInstanceID : GetMeshDescription()->GetPolygonPerimeterVertexInstances( PolygonID ) )
 						{
 							FOsdFVarVertexData& FVarVertexData = *new( FVarVertexDatas ) FOsdFVarVertexData();
-							FVarVertexData.TextureCoordinates[ 0 ] = TextureCoordinateCount > 0 ? VertexUVs.GetArrayForIndex( 0 )[ VertexInstanceID ] : FVector2D::ZeroVector;
-							FVarVertexData.TextureCoordinates[ 1 ] = TextureCoordinateCount > 1 ? VertexUVs.GetArrayForIndex( 1 )[ VertexInstanceID ] : FVector2D::ZeroVector;
+							FVarVertexData.TextureCoordinates[ 0 ] = TextureCoordinateCount > 0 ? VertexUVs.Get( VertexInstanceID, 0 ) : FVector2D::ZeroVector;
+							FVarVertexData.TextureCoordinates[ 1 ] = TextureCoordinateCount > 1 ? VertexUVs.Get( VertexInstanceID, 1 ) : FVector2D::ZeroVector;
 							FVarVertexData.VertexColor = FLinearColor( VertexColors[ VertexInstanceID ] );
 						}
 					}
@@ -2549,7 +2537,7 @@ void UEditableMesh::ComputePolygonTriangulation( const FPolygonID PolygonID, TAr
 	static TArray<FVector> VertexPositions;
 
 	{
-		const TVertexAttributeArray<FVector>& MeshVertexPositions = GetMeshDescription()->VertexAttributes().GetAttributes<FVector>( MeshAttribute::Vertex::Position );
+		TVertexAttributesConstRef<FVector> MeshVertexPositions = GetMeshDescription()->VertexAttributes().GetAttributesRef<FVector>( MeshAttribute::Vertex::Position );
 		PrevVertexNumbers.SetNumUninitialized( PolygonVertexCount, false );
 		NextVertexNumbers.SetNumUninitialized( PolygonVertexCount, false );
 		VertexPositions.SetNumUninitialized( PolygonVertexCount, false );
@@ -2665,7 +2653,7 @@ bool UEditableMesh::ComputeBarycentricWeightForPointOnPolygon( const FPolygonID
 	const FMeshPolygon& Polygon = GetMeshDescription()->GetPolygon( PolygonID );
 	const int32 TriangleCount = Polygon.Triangles.Num();
 
-	const TVertexAttributeArray<FVector>& VertexPositions = GetMeshDescription()->VertexAttributes().GetAttributes<FVector>( MeshAttribute::Vertex::Position );
+	TVertexAttributesConstRef<FVector> VertexPositions = GetMeshDescription()->VertexAttributes().GetAttributesRef<FVector>( MeshAttribute::Vertex::Position );
 
 	// Figure out which triangle the incoming point is within
 	for( const FMeshTriangle& Triangle : Polygon.Triangles )
@@ -2729,7 +2717,7 @@ void UEditableMesh::MoveVertices( const TArray<FVertexToMove>& VerticesToMove )
 	static TArray<FAttributesForVertex> VertexAttributesToSet;
 	VertexAttributesToSet.Reset();
 
-	const TVertexAttributeArray<FVector>& VertexPositions = GetMeshDescription()->VertexAttributes().GetAttributes<FVector>( MeshAttribute::Vertex::Position );
+	const TVertexAttributesRef<FVector> VertexPositions = GetMeshDescription()->VertexAttributes().GetAttributesRef<FVector>( MeshAttribute::Vertex::Position );
 
 	for( const FVertexToMove& VertexToMove : VerticesToMove )
 	{
@@ -3144,7 +3132,7 @@ void UEditableMesh::FindPolygonLoop( const FEdgeID EdgeID, TArray<FEdgeID>& OutE
 	bool bCurrentEdgeIsInOppositeDirection = false;
 	bool bCurrentEdgeIsInOppositeDirectionFromStartEdge = false;
 
-	const TVertexAttributeArray<FVector>& VertexPositions = GetMeshDescription()->VertexAttributes().GetAttributes<FVector>( MeshAttribute::Vertex::Position );
+	TVertexAttributesConstRef<FVector> VertexPositions = GetMeshDescription()->VertexAttributes().GetAttributesRef<FVector>( MeshAttribute::Vertex::Position );
 
 	for( ; ; )
 	{
@@ -4892,7 +4880,7 @@ void UEditableMesh::GetConnectedSoftEdges( const FVertexID VertexID, TArray<FEdg
 {
 	OutConnectedSoftEdges.Reset();
 
-	const TEdgeAttributeArray<bool>& EdgeHardnesses = GetMeshDescription()->EdgeAttributes().GetAttributes<bool>( MeshAttribute::Edge::IsHard );
+	TEdgeAttributesConstRef<bool> EdgeHardnesses = GetMeshDescription()->EdgeAttributes().GetAttributesRef<bool>( MeshAttribute::Edge::IsHard );
 	for( const FEdgeID ConnectedEdgeID : GetMeshDescription()->GetVertex( VertexID ).ConnectedEdgeIDs )
 	{
 		if( !EdgeHardnesses[ ConnectedEdgeID ] )
@@ -5316,7 +5304,7 @@ void UEditableMesh::ExtrudePolygons( const TArray<FPolygonID>& PolygonIDs, const
 	AttributesForVertices.Reset();
 
 	// Get vertex position attribute array
-	const TVertexAttributeArray<FVector>& VertexPositions = GetMeshDescription()->VertexAttributes().GetAttributes<FVector>( MeshAttribute::Vertex::Position );
+	const TVertexAttributesRef<FVector> VertexPositions = GetMeshDescription()->VertexAttributes().GetAttributesRef<FVector>( MeshAttribute::Vertex::Position );
 
 	// First, let's figure out which of the polygons we were asked to extrude share edges or vertices.  We'll keep those
 	// edges intact!
@@ -5622,7 +5610,7 @@ void UEditableMesh::ExtrudePolygons( const TArray<FPolygonID>& PolygonIDs, const
 			// All of the border edges of the new polygon will be hard.  If it was a shared edge, then we'll just preserve whatever was
 			// originally going on with the internal edge.
 			{
-				const TEdgeAttributeArray<bool>& EdgeHardnesses = GetMeshDescription()->EdgeAttributes().GetAttributes<bool>( MeshAttribute::Edge::IsHard );
+				const TEdgeAttributesRef<bool> EdgeHardnesses = GetMeshDescription()->EdgeAttributes().GetAttributesRef<bool>( MeshAttribute::Edge::IsHard );
 
 				const int32 NewPerimeterEdgeCount = this->GetPolygonPerimeterEdgeCount( ExtrudedFrontPolygonID );
 				check( NewPerimeterEdgeCount == GetPolygonPerimeterEdgeCount( PolygonID ) );	// New polygon should always have the same number of edges (in the same order) as the original!
@@ -5851,7 +5839,7 @@ void UEditableMesh::ExtendVertices( const TArray<FVertexID>& VertexIDs, const bo
 		static TArray<FPolygonToCreate> PolygonsToCreate;
 		PolygonsToCreate.Reset();
 
-		const TVertexAttributeArray<FVector>& VertexPositions = GetMeshDescription()->VertexAttributes().GetAttributes<FVector>( MeshAttribute::Vertex::Position );
+		const TVertexAttributesRef<FVector> VertexPositions = GetMeshDescription()->VertexAttributes().GetAttributesRef<FVector>( MeshAttribute::Vertex::Position );
 
 		for( int32 VertexNumber = 0; VertexNumber < VertexIDs.Num(); ++VertexNumber )
 		{
@@ -6032,7 +6020,7 @@ void UEditableMesh::BevelOrInsetPolygons( const TArray<FPolygonID>& PolygonIDs,
 	static TArray<FAttributesForVertex> AttributesForVertices;
 	AttributesForVertices.Reset();
 
-	const TVertexAttributeArray<FVector>& VertexPositions = GetMeshDescription()->VertexAttributes().GetAttributes<FVector>( MeshAttribute::Vertex::Position );
+	const TVertexAttributesRef<FVector> VertexPositions = GetMeshDescription()->VertexAttributes().GetAttributesRef<FVector>( MeshAttribute::Vertex::Position );
 
 	for( const FPolygonID PolygonID : PolygonIDs )
 	{
@@ -6312,7 +6300,7 @@ float UEditableMesh::GetPolygonCornerAngleForVertex( const FPolygonID PolygonID,
 		const FVertexID ThisVertexID = GetMeshDescription()->GetVertexInstanceVertex( Contour.VertexInstanceIDs[ ContourIndex ] );
 		const FVertexID NextVertexID = GetMeshDescription()->GetVertexInstanceVertex( Contour.VertexInstanceIDs[ NextIndex ] );
 
-		const TVertexAttributeArray<FVector>& VertexPositions = GetMeshDescription()->VertexAttributes().GetAttributes<FVector>( MeshAttribute::Vertex::Position );
+		TVertexAttributesConstRef<FVector> VertexPositions = GetMeshDescription()->VertexAttributes().GetAttributesRef<FVector>( MeshAttribute::Vertex::Position );
 
 		const FVector PrevVertexPosition = VertexPositions[ PrevVertexID ];
 		const FVector ThisVertexPosition = VertexPositions[ ThisVertexID ];
@@ -6358,15 +6346,15 @@ float UEditableMesh::GetPolygonCornerAngleForVertex( const FPolygonID PolygonID,
 
 void UEditableMesh::GeneratePolygonTangentsAndNormals( const TArray<FPolygonID>& PolygonIDs )
 {
-	const TVertexAttributeArray<FVector>& VertexPositions = GetMeshDescription()->VertexAttributes().GetAttributes<FVector>( MeshAttribute::Vertex::Position );
+	const TVertexAttributesRef<FVector> VertexPositions = GetMeshDescription()->VertexAttributes().GetAttributesRef<FVector>( MeshAttribute::Vertex::Position );
 
 	// @todo mesheditor: currently hardcoded to calculate the tangent basis from UV0.
-	const TVertexInstanceAttributeArray<FVector2D>& VertexUVs = GetMeshDescription()->VertexInstanceAttributes().GetAttributes<FVector2D>( MeshAttribute::VertexInstance::TextureCoordinate, 0 );
+	const TVertexInstanceAttributesRef<FVector2D> VertexUVs = GetMeshDescription()->VertexInstanceAttributes().GetAttributesRef<FVector2D>( MeshAttribute::VertexInstance::TextureCoordinate );
 
-	TPolygonAttributeArray<FVector>& PolygonNormals = GetMeshDescription()->PolygonAttributes().GetAttributes<FVector>( MeshAttribute::Polygon::Normal );
-	TPolygonAttributeArray<FVector>& PolygonTangents = GetMeshDescription()->PolygonAttributes().GetAttributes<FVector>( MeshAttribute::Polygon::Tangent );
-	TPolygonAttributeArray<FVector>& PolygonBinormals = GetMeshDescription()->PolygonAttributes().GetAttributes<FVector>( MeshAttribute::Polygon::Binormal );
-	TPolygonAttributeArray<FVector>& PolygonCenters = GetMeshDescription()->PolygonAttributes().GetAttributes<FVector>( MeshAttribute::Polygon::Center );
+	TPolygonAttributesRef<FVector> PolygonNormals = GetMeshDescription()->PolygonAttributes().GetAttributesRef<FVector>( MeshAttribute::Polygon::Normal );
+	TPolygonAttributesRef<FVector> PolygonTangents = GetMeshDescription()->PolygonAttributes().GetAttributesRef<FVector>( MeshAttribute::Polygon::Tangent );
+	TPolygonAttributesRef<FVector> PolygonBinormals = GetMeshDescription()->PolygonAttributes().GetAttributesRef<FVector>( MeshAttribute::Polygon::Binormal );
+	TPolygonAttributesRef<FVector> PolygonCenters = GetMeshDescription()->PolygonAttributes().GetAttributesRef<FVector>( MeshAttribute::Polygon::Center );
 
 	for( const FPolygonID PolygonID : PolygonIDs )
 	{
@@ -6393,8 +6381,8 @@ void UEditableMesh::GeneratePolygonTangentsAndNormals( const TArray<FPolygonID>&
 			const FVector DPosition1 = VertexPositions[ VertexID1 ] - VertexPositions[ VertexID0 ];
 			const FVector DPosition2 = VertexPositions[ VertexID2 ] - VertexPositions[ VertexID0 ];
 
-			const FVector2D DUV1 = VertexUVs[ Triangle.VertexInstanceID1 ] - VertexUVs[ Triangle.VertexInstanceID0 ];
-			const FVector2D DUV2 = VertexUVs[ Triangle.VertexInstanceID2 ] - VertexUVs[ Triangle.VertexInstanceID0 ];
+			const FVector2D DUV1 = VertexUVs.Get( Triangle.VertexInstanceID1, 0 ) - VertexUVs.Get( Triangle.VertexInstanceID0, 0 );
+			const FVector2D DUV2 = VertexUVs.Get( Triangle.VertexInstanceID2, 0 ) - VertexUVs.Get( Triangle.VertexInstanceID0, 0 );
 
 			// We have a left-handed coordinate system, but a counter-clockwise winding order
 			// Hence normal calculation has to take the triangle vectors cross product in reverse.
@@ -6418,9 +6406,9 @@ void UEditableMesh::GeneratePolygonTangentsAndNormals( const TArray<FPolygonID>&
 
 void UEditableMesh::GenerateTangentsAndNormals()
 {
-	const TPolygonAttributeArray<FVector>& PolygonNormals = GetMeshDescription()->PolygonAttributes().GetAttributes<FVector>( MeshAttribute::Polygon::Normal );
-	const TPolygonAttributeArray<FVector>& PolygonTangents = GetMeshDescription()->PolygonAttributes().GetAttributes<FVector>( MeshAttribute::Polygon::Tangent );
-	const TPolygonAttributeArray<FVector>& PolygonBinormals = GetMeshDescription()->PolygonAttributes().GetAttributes<FVector>( MeshAttribute::Polygon::Binormal );
+	const TPolygonAttributesRef<FVector> PolygonNormals = GetMeshDescription()->PolygonAttributes().GetAttributesRef<FVector>( MeshAttribute::Polygon::Normal );
+	const TPolygonAttributesRef<FVector> PolygonTangents = GetMeshDescription()->PolygonAttributes().GetAttributesRef<FVector>( MeshAttribute::Polygon::Tangent );
+	const TPolygonAttributesRef<FVector> PolygonBinormals = GetMeshDescription()->PolygonAttributes().GetAttributesRef<FVector>( MeshAttribute::Polygon::Binormal );
 
 	GeneratePolygonTangentsAndNormals( PolygonsPendingNewTangentBasis.Array() );
 
@@ -7338,7 +7326,7 @@ void UEditableMesh::WeldVertices( const TArray<FVertexID>& VertexIDsToWeld, FVer
 		return;
 	}
 
-	const TVertexAttributeArray<FVector>& VertexPositions = GetMeshDescription()->VertexAttributes().GetAttributes<FVector>( MeshAttribute::Vertex::Position );
+	const TVertexAttributesRef<FVector> VertexPositions = GetMeshDescription()->VertexAttributes().GetAttributesRef<FVector>( MeshAttribute::Vertex::Position );
 
 	// Create new welded vertex
 	static TArray<FVertexID> NewVertices;
@@ -7491,7 +7479,7 @@ void UEditableMesh::TessellatePolygons( const TArray<FPolygonID>& PolygonIDs, co
 
 	OutNewPolygonIDs.Reset();
 
-	const TVertexAttributeArray<FVector>& VertexPositions = GetMeshDescription()->VertexAttributes().GetAttributes<FVector>( MeshAttribute::Vertex::Position );
+	const TVertexAttributesRef<FVector> VertexPositions = GetMeshDescription()->VertexAttributes().GetAttributesRef<FVector>( MeshAttribute::Vertex::Position );
 
 	//
 	// Simple tessellation algorithm:
@@ -7925,8 +7913,8 @@ void UEditableMesh::QuadrangulatePolygonGroup( const FPolygonGroupID PolygonGrou
 	{
 		float BestScore = TNumericLimits<float>::Max();
 
-		const TEdgeAttributeArray<bool>& EdgeHardnesses = GetMeshDescription()->EdgeAttributes().GetAttributes<bool>( MeshAttribute::Edge::IsHard );
-		const TVertexAttributeArray<FVector>& VertexPositions = GetMeshDescription()->VertexAttributes().GetAttributes<FVector>( MeshAttribute::Vertex::Position );
+		const TEdgeAttributesRef<bool> EdgeHardnesses = GetMeshDescription()->EdgeAttributes().GetAttributesRef<bool>( MeshAttribute::Edge::IsHard );
+		const TVertexAttributesRef<FVector> VertexPositions = GetMeshDescription()->VertexAttributes().GetAttributesRef<FVector>( MeshAttribute::Vertex::Position );
 
 		for( const FPolygonID PolygonID : PolygonIDs )
 		{
@@ -8314,7 +8302,7 @@ void UEditableMesh::UpdateOrRebuildOctree()
 
 			// Now, add new polygons to the octree
 			{
-				const TVertexAttributeArray<FVector>& VertexPositions = GetMeshDescription()->VertexAttributes().GetAttributes<FVector>( MeshAttribute::Vertex::Position );
+				const TVertexAttributesRef<FVector> VertexPositions = GetMeshDescription()->VertexAttributes().GetAttributesRef<FVector>( MeshAttribute::Vertex::Position );
 				for( const FPolygonID PolygonID : NewOctreePolygonIDs )
 				{
 					checkSlow( IsValidPolygon( PolygonID ) );
@@ -8383,7 +8371,7 @@ void UEditableMesh::RebuildOctree()
 		// @todo mesheditor spatial: We'll currently only ever find vertices and edges that are connected to polygons.  Our spatial
 		// database only tracks polygons.  Do we need to support hover/selection of "loose" vertices and edges also?
 
-		const TVertexAttributeArray<FVector>& VertexPositions = GetMeshDescription()->VertexAttributes().GetAttributes<FVector>( MeshAttribute::Vertex::Position );
+		const TVertexAttributesRef<FVector> VertexPositions = GetMeshDescription()->VertexAttributes().GetAttributesRef<FVector>( MeshAttribute::Vertex::Position );
 		for( const FPolygonID PolygonID : GetMeshDescription()->Polygons().GetElementIDs() )
 		{
 			FBox BoundingBox;
diff --git a/Engine/Plugins/Runtime/EditableMesh/Source/EditableMesh/EditableStaticMeshAdapter.cpp b/Engine/Plugins/Runtime/EditableMesh/Source/EditableMesh/EditableStaticMeshAdapter.cpp
--- a/Engine/Plugins/Runtime/EditableMesh/Source/EditableMesh/EditableStaticMeshAdapter.cpp
+++ b/Engine/Plugins/Runtime/EditableMesh/Source/EditableMesh/EditableStaticMeshAdapter.cpp
@@ -94,20 +94,20 @@ void UEditableStaticMeshAdapter::InitEditableStaticMesh( UEditableMesh* Editable
 
 				FMeshDescription* MeshDescription = EditableMesh->GetMeshDescription();
 
-				TVertexAttributeArray<FVector>& VertexPositions = MeshDescription->VertexAttributes().GetAttributes<FVector>( MeshAttribute::Vertex::Position );
+				TVertexAttributesRef<FVector> VertexPositions = MeshDescription->VertexAttributes().GetAttributesRef<FVector>( MeshAttribute::Vertex::Position );
 
-				TVertexInstanceAttributeArray<FVector>& VertexInstanceNormals = MeshDescription->VertexInstanceAttributes().GetAttributes<FVector>( MeshAttribute::VertexInstance::Normal );
-				TVertexInstanceAttributeArray<FVector>& VertexInstanceTangents = MeshDescription->VertexInstanceAttributes().GetAttributes<FVector>( MeshAttribute::VertexInstance::Tangent );
-				TVertexInstanceAttributeArray<float>& VertexInstanceBinormalSigns = MeshDescription->VertexInstanceAttributes().GetAttributes<float>( MeshAttribute::VertexInstance::BinormalSign );
-				TVertexInstanceAttributeArray<FVector4>& VertexInstanceColors = MeshDescription->VertexInstanceAttributes().GetAttributes<FVector4>( MeshAttribute::VertexInstance::Color );
-				TVertexInstanceAttributeIndicesArray<FVector2D>& VertexInstanceUVs = MeshDescription->VertexInstanceAttributes().GetAttributesSet<FVector2D>( MeshAttribute::VertexInstance::TextureCoordinate );
+				TVertexInstanceAttributesRef<FVector> VertexInstanceNormals = MeshDescription->VertexInstanceAttributes().GetAttributesRef<FVector>( MeshAttribute::VertexInstance::Normal );
+				TVertexInstanceAttributesRef<FVector> VertexInstanceTangents = MeshDescription->VertexInstanceAttributes().GetAttributesRef<FVector>( MeshAttribute::VertexInstance::Tangent );
+				TVertexInstanceAttributesRef<float> VertexInstanceBinormalSigns = MeshDescription->VertexInstanceAttributes().GetAttributesRef<float>( MeshAttribute::VertexInstance::BinormalSign );
+				TVertexInstanceAttributesRef<FVector4> VertexInstanceColors = MeshDescription->VertexInstanceAttributes().GetAttributesRef<FVector4>( MeshAttribute::VertexInstance::Color );
+				TVertexInstanceAttributesRef<FVector2D> VertexInstanceUVs = MeshDescription->VertexInstanceAttributes().GetAttributesRef<FVector2D>( MeshAttribute::VertexInstance::TextureCoordinate );
 
-				TEdgeAttributeArray<bool>& EdgeHardnesses = MeshDescription->EdgeAttributes().GetAttributes<bool>( MeshAttribute::Edge::IsHard );
+				TEdgeAttributesRef<bool> EdgeHardnesses = MeshDescription->EdgeAttributes().GetAttributesRef<bool>( MeshAttribute::Edge::IsHard );
 
-				TPolygonGroupAttributeArray<FName>& PolygonGroupImportedMaterialSlotNames = MeshDescription->PolygonGroupAttributes().GetAttributes<FName>( MeshAttribute::PolygonGroup::ImportedMaterialSlotName );
-				TPolygonGroupAttributeArray<FName>& PolygonGroupMaterialAssetNames = MeshDescription->PolygonGroupAttributes().GetAttributes<FName>( MeshAttribute::PolygonGroup::MaterialAssetName );
-				TPolygonGroupAttributeArray<bool>& PolygonGroupCollision = MeshDescription->PolygonGroupAttributes().GetAttributes<bool>( MeshAttribute::PolygonGroup::EnableCollision );
-				TPolygonGroupAttributeArray<bool>& PolygonGroupCastShadow = MeshDescription->PolygonGroupAttributes().GetAttributes<bool>( MeshAttribute::PolygonGroup::CastShadow );
+				TPolygonGroupAttributesRef<FName> PolygonGroupImportedMaterialSlotNames = MeshDescription->PolygonGroupAttributes().GetAttributesRef<FName>( MeshAttribute::PolygonGroup::ImportedMaterialSlotName );
+				TPolygonGroupAttributesRef<FName> PolygonGroupMaterialAssetNames = MeshDescription->PolygonGroupAttributes().GetAttributesRef<FName>( MeshAttribute::PolygonGroup::MaterialAssetName );
+				TPolygonGroupAttributesRef<bool> PolygonGroupCollision = MeshDescription->PolygonGroupAttributes().GetAttributesRef<bool>( MeshAttribute::PolygonGroup::EnableCollision );
+				TPolygonGroupAttributesRef<bool> PolygonGroupCastShadow = MeshDescription->PolygonGroupAttributes().GetAttributesRef<bool>( MeshAttribute::PolygonGroup::CastShadow );
 
 				// Store off the number of texture coordinates in this mesh
 				const int32 NumUVs = StaticMeshLOD.GetNumTexCoords();
@@ -252,7 +252,7 @@ void UEditableStaticMeshAdapter::InitEditableStaticMesh( UEditableMesh* Editable
 						VertexInstanceColors[ VertexInstanceID ] = Color;
 						for( int32 UVIndex = 0; UVIndex < NumUVs; ++UVIndex )
 						{
-							VertexInstanceUVs.GetArrayForIndex( UVIndex )[ VertexInstanceID ] = StaticMeshLOD.VertexBuffers.StaticMeshVertexBuffer.GetVertexUV( RenderingVertexIndex, UVIndex );
+							VertexInstanceUVs.Set( VertexInstanceID, UVIndex, StaticMeshLOD.VertexBuffers.StaticMeshVertexBuffer.GetVertexUV( RenderingVertexIndex, UVIndex ) );
 						}
 					}
 				}
@@ -506,9 +506,9 @@ void UEditableStaticMeshAdapter::OnRebuildRenderMesh( const UEditableMesh* Edita
 	const FMeshDescription* MeshDescription = EditableMesh->GetMeshDescription();
 	check( MeshDescription );
 
-	const TPolygonGroupAttributeArray<FName>& PolygonGroupImportedMaterialSlotNames = MeshDescription->PolygonGroupAttributes().GetAttributes<FName>( MeshAttribute::PolygonGroup::ImportedMaterialSlotName );
-	const TPolygonGroupAttributeArray<bool>& PolygonGroupCollision = MeshDescription->PolygonGroupAttributes().GetAttributes<bool>( MeshAttribute::PolygonGroup::EnableCollision );
-	const TPolygonGroupAttributeArray<bool>& PolygonGroupCastShadow = MeshDescription->PolygonGroupAttributes().GetAttributes<bool>( MeshAttribute::PolygonGroup::CastShadow );
+	TPolygonGroupAttributesConstRef<FName> PolygonGroupImportedMaterialSlotNames = MeshDescription->PolygonGroupAttributes().GetAttributesRef<FName>( MeshAttribute::PolygonGroup::ImportedMaterialSlotName );
+	TPolygonGroupAttributesConstRef<bool> PolygonGroupCollision = MeshDescription->PolygonGroupAttributes().GetAttributesRef<bool>( MeshAttribute::PolygonGroup::EnableCollision );
+	TPolygonGroupAttributesConstRef<bool> PolygonGroupCastShadow = MeshDescription->PolygonGroupAttributes().GetAttributesRef<bool>( MeshAttribute::PolygonGroup::CastShadow );
 
 	const int32 NumPolygonGroups = MeshDescription->PolygonGroups().Num();
 	StaticMeshLOD.Sections.Empty( NumPolygonGroups );
@@ -614,12 +614,12 @@ void UEditableStaticMeshAdapter::OnRebuildRenderMesh( const UEditableMesh* Edita
 		// set up vertex buffer elements
 		StaticMeshBuildVertices.SetNum( MeshDescription->VertexInstances().GetArraySize() );
 
-		const TVertexAttributeArray<FVector>& VertexPositions = MeshDescription->VertexAttributes().GetAttributes<FVector>( MeshAttribute::Vertex::Position );
-		const TVertexInstanceAttributeArray<FVector>& VertexInstanceNormals = MeshDescription->VertexInstanceAttributes().GetAttributes<FVector>( MeshAttribute::VertexInstance::Normal );
-		const TVertexInstanceAttributeArray<FVector>& VertexInstanceTangents = MeshDescription->VertexInstanceAttributes().GetAttributes<FVector>( MeshAttribute::VertexInstance::Tangent );
-		const TVertexInstanceAttributeArray<float>& VertexInstanceBinormalSigns = MeshDescription->VertexInstanceAttributes().GetAttributes<float>( MeshAttribute::VertexInstance::BinormalSign );
-		const TVertexInstanceAttributeArray<FVector4>& VertexInstanceColors = MeshDescription->VertexInstanceAttributes().GetAttributes<FVector4>( MeshAttribute::VertexInstance::Color );
-		const TVertexInstanceAttributeIndicesArray<FVector2D>& VertexInstanceUVs = MeshDescription->VertexInstanceAttributes().GetAttributesSet<FVector2D>( MeshAttribute::VertexInstance::TextureCoordinate );
+		TVertexAttributesConstRef<FVector> VertexPositions = MeshDescription->VertexAttributes().GetAttributesRef<FVector>( MeshAttribute::Vertex::Position );
+		TVertexInstanceAttributesConstRef<FVector> VertexInstanceNormals = MeshDescription->VertexInstanceAttributes().GetAttributesRef<FVector>( MeshAttribute::VertexInstance::Normal );
+		TVertexInstanceAttributesConstRef<FVector> VertexInstanceTangents = MeshDescription->VertexInstanceAttributes().GetAttributesRef<FVector>( MeshAttribute::VertexInstance::Tangent );
+		TVertexInstanceAttributesConstRef<float> VertexInstanceBinormalSigns = MeshDescription->VertexInstanceAttributes().GetAttributesRef<float>( MeshAttribute::VertexInstance::BinormalSign );
+		TVertexInstanceAttributesConstRef<FVector4> VertexInstanceColors = MeshDescription->VertexInstanceAttributes().GetAttributesRef<FVector4>( MeshAttribute::VertexInstance::Color );
+		TVertexInstanceAttributesConstRef<FVector2D> VertexInstanceUVs = MeshDescription->VertexInstanceAttributes().GetAttributesRef<FVector2D>( MeshAttribute::VertexInstance::TextureCoordinate );
 
 		for( const FVertexInstanceID VertexInstanceID : MeshDescription->VertexInstances().GetElementIDs() )
 		{
@@ -639,7 +639,7 @@ void UEditableStaticMeshAdapter::OnRebuildRenderMesh( const UEditableMesh* Edita
 			StaticMeshVertex.Color = VertexInstanceColor.ToFColor( true );
 			for( int32 UVIndex = 0; UVIndex < VertexInstanceUVs.GetNumIndices(); ++UVIndex )
 			{
-				StaticMeshVertex.UVs[ UVIndex ] = VertexInstanceUVs.GetArrayForIndex( UVIndex )[ VertexInstanceID ];
+				StaticMeshVertex.UVs[ UVIndex ] = VertexInstanceUVs.Get( VertexInstanceID, UVIndex );
 			}
 		}
 
@@ -845,7 +845,7 @@ void UEditableStaticMeshAdapter::OnReindexElements( const UEditableMesh* Editabl
 		FRenderingPolygonGroup& RenderingPolygonGroup = RenderingPolygonGroups[ PolygonGroupID ];
 //		const FMeshPolygonGroup& PolygonGroup = EditableMesh->GetMeshDescription()->GetPolygonGroup( PolygonGroupID );
 
-		TSparseArray<FTriangleID> TriangleRemappings;
+		TSparseArray<int32> TriangleRemappings;
 		RenderingPolygonGroup.Triangles.Compact( TriangleRemappings );
 
 		for( const FTriangleID TriangleID : RenderingPolygonGroup.Triangles.GetElementIDs() )
@@ -854,7 +854,7 @@ void UEditableStaticMeshAdapter::OnReindexElements( const UEditableMesh* Editabl
 			for( int32 TriangleVertexNumber = 0; TriangleVertexNumber < 3; ++TriangleVertexNumber )
 			{
 				const FVertexInstanceID OriginalVertexInstanceID = Triangle.GetVertexInstanceID( TriangleVertexNumber );
-				const FVertexInstanceID NewVertexInstanceID = Remappings.NewVertexInstanceIndexLookup[ OriginalVertexInstanceID.GetValue() ];
+				const FVertexInstanceID NewVertexInstanceID( Remappings.NewVertexInstanceIndexLookup[ OriginalVertexInstanceID.GetValue() ] );
 				Triangle.SetVertexInstanceID( TriangleVertexNumber, NewVertexInstanceID );
 			}
 		}
@@ -869,7 +869,7 @@ void UEditableStaticMeshAdapter::OnReindexElements( const UEditableMesh* Editabl
 
 			for( FTriangleID& TriangleID : RenderingPolygon.TriangulatedPolygonTriangleIndices )
 			{
-				TriangleID = TriangleRemappings[ TriangleID.GetValue() ];
+				TriangleID = FTriangleID( TriangleRemappings[ TriangleID.GetValue() ] );
 			}
 		}
 	}
@@ -994,7 +994,7 @@ void UEditableStaticMeshAdapter::UpdateBounds( const UEditableMesh* EditableMesh
 
 			const FMeshDescription* MeshDescription = EditableMesh->GetMeshDescription();
 
-			const TVertexAttributeArray<FVector>& VertexPositions = EditableMesh->GetMeshDescription()->VertexAttributes().GetAttributes<FVector>( MeshAttribute::Vertex::Position );
+			TVertexAttributesConstRef<FVector> VertexPositions = EditableMesh->GetMeshDescription()->VertexAttributes().GetAttributesRef<FVector>( MeshAttribute::Vertex::Position );
 
 			for( const FVertexID VertexID : MeshDescription->Vertices().GetElementIDs() )
 			{
@@ -1222,7 +1222,7 @@ void UEditableStaticMeshAdapter::OnCreateVertexInstances( const UEditableMesh* E
 		const FMeshDescription* MeshDescription = EditableMesh->GetMeshDescription();
 		FStaticMeshLODResources& StaticMeshLOD = GetStaticMeshLOD();
 		const int32 NumUVs = MeshDescription->VertexInstanceAttributes().GetAttributeIndexCount<FVector2D>( MeshAttribute::VertexInstance::TextureCoordinate );
-		const TVertexAttributeArray<FVector>& VertexPositions = MeshDescription->VertexAttributes().GetAttributes<FVector>( MeshAttribute::Vertex::Position );
+		TVertexAttributesConstRef<FVector> VertexPositions = MeshDescription->VertexAttributes().GetAttributesRef<FVector>( MeshAttribute::Vertex::Position );
 		const bool bHasColors = StaticMeshLOD.VertexBuffers.ColorVertexBuffer.GetNumVertices() > 0;
 
 		// Determine if we need to grow the render buffers
@@ -1638,10 +1638,10 @@ void UEditableStaticMeshAdapter::OnCreatePolygonGroups( const UEditableMesh* Edi
 	const FMeshDescription* MeshDescription = EditableMesh->GetMeshDescription();
 	check( MeshDescription );
 
-	const TPolygonGroupAttributeArray<FName>& PolygonGroupImportedMaterialSlotNames = MeshDescription->PolygonGroupAttributes().GetAttributes<FName>( MeshAttribute::PolygonGroup::ImportedMaterialSlotName );
-	const TPolygonGroupAttributeArray<FName>& PolygonGroupMaterialAssetSlotNames = MeshDescription->PolygonGroupAttributes().GetAttributes<FName>( MeshAttribute::PolygonGroup::MaterialAssetName );
-	const TPolygonGroupAttributeArray<bool>& PolygonGroupCollision = MeshDescription->PolygonGroupAttributes().GetAttributes<bool>( MeshAttribute::PolygonGroup::EnableCollision );
-	const TPolygonGroupAttributeArray<bool>& PolygonGroupCastShadow = MeshDescription->PolygonGroupAttributes().GetAttributes<bool>( MeshAttribute::PolygonGroup::CastShadow );
+	TPolygonGroupAttributesConstRef<FName> PolygonGroupImportedMaterialSlotNames = MeshDescription->PolygonGroupAttributes().GetAttributesRef<FName>( MeshAttribute::PolygonGroup::ImportedMaterialSlotName );
+	TPolygonGroupAttributesConstRef<FName> PolygonGroupMaterialAssetSlotNames = MeshDescription->PolygonGroupAttributes().GetAttributesRef<FName>( MeshAttribute::PolygonGroup::MaterialAssetName );
+	TPolygonGroupAttributesConstRef<bool> PolygonGroupCollision = MeshDescription->PolygonGroupAttributes().GetAttributesRef<bool>( MeshAttribute::PolygonGroup::EnableCollision );
+	TPolygonGroupAttributesConstRef<bool> PolygonGroupCastShadow = MeshDescription->PolygonGroupAttributes().GetAttributesRef<bool>( MeshAttribute::PolygonGroup::CastShadow );
 
 	for( const FPolygonGroupID PolygonGroupID : PolygonGroupIDs )
 	{
@@ -1746,7 +1746,7 @@ void UEditableStaticMeshAdapter::OnSetPolygonGroupAttribute( const UEditableMesh
 
 void UEditableStaticMeshAdapter::OnDeletePolygonGroups( const UEditableMesh* EditableMesh, const TArray<FPolygonGroupID>& PolygonGroupIDs )
 {
-	const TPolygonGroupAttributeArray<FName>& PolygonGroupImportedMaterialSlotNames = EditableMesh->GetMeshDescription()->PolygonGroupAttributes().GetAttributes<FName>( MeshAttribute::PolygonGroup::ImportedMaterialSlotName );
+	TPolygonGroupAttributesConstRef<FName> PolygonGroupImportedMaterialSlotNames = EditableMesh->GetMeshDescription()->PolygonGroupAttributes().GetAttributesRef<FName>( MeshAttribute::PolygonGroup::ImportedMaterialSlotName );
 
 	for( const FPolygonGroupID PolygonGroupID : PolygonGroupIDs )
 	{
diff --git a/Engine/Plugins/Tests/EditorTests/Source/EditorTests/Private/UnrealEd/MeshDescriptionTests.cpp b/Engine/Plugins/Tests/EditorTests/Source/EditorTests/Private/UnrealEd/MeshDescriptionTests.cpp
--- a/Engine/Plugins/Tests/EditorTests/Source/EditorTests/Private/UnrealEd/MeshDescriptionTests.cpp
+++ b/Engine/Plugins/Tests/EditorTests/Source/EditorTests/Private/UnrealEd/MeshDescriptionTests.cpp
@@ -288,6 +288,74 @@ bool FMeshDescriptionTest::CompareRawMesh(const FString& AssetName, FAutomationT
 	return bAllSame;
 }
 
+
+template <typename T> FString AttributeValueAsString(const T& Value) { return LexToString(Value); }
+FString AttributeValueAsString(const FVector& Value) { return Value.ToString(); }
+FString AttributeValueAsString(const FVector2D& Value) { return Value.ToString(); }
+FString AttributeValueAsString(const FVector4& Value) { return Value.ToString(); }
+
+template<typename T, typename U>
+void MeshDescriptionAttributeArrayCompare(const FString& ConversionName, const FString& AssetName, FAutomationTestExecutionInfo& ExecutionInfo, bool& bIsSame, const U& ElementIterator, const FString& ArrayName, const T ReferenceArray, const T ResultArray)
+{
+	if (ReferenceArray.GetNumElements() != ResultArray.GetNumElements())
+	{
+		ExecutionInfo.AddEvent(FAutomationEvent(EAutomationEventType::Error, FString::Printf(TEXT("The %s conversion %s is not lossless, %s count is different. %s count expected [%d] result [%d]"),
+			*AssetName,
+			*ConversionName,
+			*ArrayName,
+			*ArrayName,
+			ReferenceArray.GetNumElements(),
+			ResultArray.GetNumElements())));
+		bIsSame = false;
+	}
+	else if (ReferenceArray.GetNumIndices() != ReferenceArray.GetNumIndices())
+	{
+		ExecutionInfo.AddEvent(FAutomationEvent(EAutomationEventType::Error, FString::Printf(TEXT("The %s conversion %s is not lossless, %s channel count is different. %s channel count expected [%d] result [%d]"),
+			*AssetName,
+			*ConversionName,
+			*ArrayName,
+			*ArrayName,
+			ReferenceArray.GetNumIndices(),
+			ResultArray.GetNumIndices())));
+		bIsSame = false;
+
+	}
+	else
+	{
+		int32 NumDifferent = 0;
+		for (int32 Index = 0; Index < ReferenceArray.GetNumIndices(); ++Index)
+		{
+			for (auto ElementID : ElementIterator.GetElementIDs())
+			{
+				if (ReferenceArray.Get(ElementID, Index) != ResultArray.Get(ElementID, Index))
+				{
+					NumDifferent++;
+					if (NumDifferent < 5)
+					{
+						ExecutionInfo.AddEvent(FAutomationEvent(EAutomationEventType::Error, FString::Printf(TEXT("The %s conversion %s is not lossless, %s array is different. Element [%d] of attribute index [%d] expected %s [%s] result [%s]"),
+							*AssetName,
+							*ConversionName,
+							*ArrayName,
+							ElementID.GetValue(),
+							Index,
+							*ArrayName,
+							*AttributeValueAsString(ReferenceArray.Get(ElementID, Index)),
+							*AttributeValueAsString(ResultArray.Get(ElementID, Index)))));
+					}
+					else
+					{
+						ExecutionInfo.AddEvent(FAutomationEvent(EAutomationEventType::Error, TEXT("More than 5 unequal elements - silencing")));
+						break;
+					}
+					bIsSame = false;
+				}
+			}
+		}
+	}
+}
+
+
+/*
 template<typename T, typename U, typename V>
 void MeshDescriptionStructureArrayCompare(const FString& ConversionName, const FString& AssetName, FAutomationTestExecutionInfo& ExecutionInfo, bool& bIsSame, const U& ElementIterator, const FString& VectorArrayName, const T& ReferenceArray, const T& ResultArray)
 {
@@ -324,7 +392,7 @@ void MeshDescriptionStructureArrayCompare(const FString& ConversionName, const F
 }
 
 template<typename T, typename U, typename V>
-void MeshDescriptionNumberArrayCompare(const FString& ConversionName, const FString& AssetName, FAutomationTestExecutionInfo& ExecutionInfo, bool& bIsSame, const U& ElementIterator, const FString& VectorArrayName, const T& ReferenceArray, const T& ResultArray)
+void MeshDescriptionAttributeArrayCompare(const FString& ConversionName, const FString& AssetName, FAutomationTestExecutionInfo& ExecutionInfo, bool& bIsSame, const U& ElementIterator, const FString& VectorArrayName, const T& ReferenceArray, const T& ResultArray)
 {
 	if (ReferenceArray.Num() != ResultArray.Num())
 	{
@@ -392,37 +460,32 @@ void MeshDescriptionNumberFNameCompare(const FString& ConversionName, const FStr
 		}
 	}
 }
+*/
 
 bool FMeshDescriptionTest::CompareMeshDescription(const FString& AssetName, FAutomationTestExecutionInfo& ExecutionInfo, const FMeshDescription& ReferenceMeshDescription, const FMeshDescription& MeshDescription) const
 {
 	//////////////////////////////////////////////////////////////////////////
 	//Gather the reference data
-	const TVertexAttributeArray<FVector>& ReferenceVertexPositions = ReferenceMeshDescription.VertexAttributes().GetAttributes<FVector>(MeshAttribute::Vertex::Position);
-
-	const TVertexInstanceAttributeArray<FVector>& ReferenceVertexInstanceNormals = ReferenceMeshDescription.VertexInstanceAttributes().GetAttributes<FVector>(MeshAttribute::VertexInstance::Normal);
-	const TVertexInstanceAttributeArray<FVector>& ReferenceVertexInstanceTangents = ReferenceMeshDescription.VertexInstanceAttributes().GetAttributes<FVector>(MeshAttribute::VertexInstance::Tangent);
-	const TVertexInstanceAttributeArray<float>& ReferenceVertexInstanceBinormalSigns = ReferenceMeshDescription.VertexInstanceAttributes().GetAttributes<float>(MeshAttribute::VertexInstance::BinormalSign);
-	const TVertexInstanceAttributeArray<FVector4>& ReferenceVertexInstanceColors = ReferenceMeshDescription.VertexInstanceAttributes().GetAttributes<FVector4>(MeshAttribute::VertexInstance::Color);
-	const TVertexInstanceAttributeIndicesArray<FVector2D>& ReferenceVertexInstanceUVs = ReferenceMeshDescription.VertexInstanceAttributes().GetAttributesSet<FVector2D>(MeshAttribute::VertexInstance::TextureCoordinate);
-
-	const TEdgeAttributeArray<bool>& ReferenceEdgeHardnesses = ReferenceMeshDescription.EdgeAttributes().GetAttributes<bool>(MeshAttribute::Edge::IsHard);
-
-	const TPolygonGroupAttributeArray<FName>& ReferencePolygonGroupMaterialName = ReferenceMeshDescription.PolygonGroupAttributes().GetAttributes<FName>(MeshAttribute::PolygonGroup::ImportedMaterialSlotName);
+	TVertexAttributesConstRef<FVector> ReferenceVertexPositions = ReferenceMeshDescription.VertexAttributes().GetAttributesRef<FVector>(MeshAttribute::Vertex::Position);
+	TVertexInstanceAttributesConstRef<FVector> ReferenceVertexInstanceNormals = ReferenceMeshDescription.VertexInstanceAttributes().GetAttributesRef<FVector>(MeshAttribute::VertexInstance::Normal);
+	TVertexInstanceAttributesConstRef<FVector> ReferenceVertexInstanceTangents = ReferenceMeshDescription.VertexInstanceAttributes().GetAttributesRef<FVector>(MeshAttribute::VertexInstance::Tangent);
+	TVertexInstanceAttributesConstRef<float> ReferenceVertexInstanceBinormalSigns = ReferenceMeshDescription.VertexInstanceAttributes().GetAttributesRef<float>(MeshAttribute::VertexInstance::BinormalSign);
+	TVertexInstanceAttributesConstRef<FVector4> ReferenceVertexInstanceColors = ReferenceMeshDescription.VertexInstanceAttributes().GetAttributesRef<FVector4>(MeshAttribute::VertexInstance::Color);
+	TVertexInstanceAttributesConstRef<FVector2D> ReferenceVertexInstanceUVs = ReferenceMeshDescription.VertexInstanceAttributes().GetAttributesRef<FVector2D>(MeshAttribute::VertexInstance::TextureCoordinate);
+	TEdgeAttributesConstRef<bool> ReferenceEdgeHardnesses = ReferenceMeshDescription.EdgeAttributes().GetAttributesRef<bool>(MeshAttribute::Edge::IsHard);
+	TPolygonGroupAttributesConstRef<FName> ReferencePolygonGroupMaterialName = ReferenceMeshDescription.PolygonGroupAttributes().GetAttributesRef<FName>(MeshAttribute::PolygonGroup::ImportedMaterialSlotName);
 
 
 	//////////////////////////////////////////////////////////////////////////
 	//Gather the result data
-	const TVertexAttributeArray<FVector>& ResultVertexPositions = MeshDescription.VertexAttributes().GetAttributes<FVector>(MeshAttribute::Vertex::Position);
-
-	const TVertexInstanceAttributeArray<FVector>& ResultVertexInstanceNormals = MeshDescription.VertexInstanceAttributes().GetAttributes<FVector>(MeshAttribute::VertexInstance::Normal);
-	const TVertexInstanceAttributeArray<FVector>& ResultVertexInstanceTangents = MeshDescription.VertexInstanceAttributes().GetAttributes<FVector>(MeshAttribute::VertexInstance::Tangent);
-	const TVertexInstanceAttributeArray<float>& ResultVertexInstanceBinormalSigns = MeshDescription.VertexInstanceAttributes().GetAttributes<float>(MeshAttribute::VertexInstance::BinormalSign);
-	const TVertexInstanceAttributeArray<FVector4>& ResultVertexInstanceColors = MeshDescription.VertexInstanceAttributes().GetAttributes<FVector4>(MeshAttribute::VertexInstance::Color);
-	const TVertexInstanceAttributeIndicesArray<FVector2D>& ResultVertexInstanceUVs = MeshDescription.VertexInstanceAttributes().GetAttributesSet<FVector2D>(MeshAttribute::VertexInstance::TextureCoordinate);
-
-	const TEdgeAttributeArray<bool>& ResultEdgeHardnesses = MeshDescription.EdgeAttributes().GetAttributes<bool>(MeshAttribute::Edge::IsHard);
-
-	const TPolygonGroupAttributeArray<FName>& ResultPolygonGroupMaterialName = MeshDescription.PolygonGroupAttributes().GetAttributes<FName>(MeshAttribute::PolygonGroup::ImportedMaterialSlotName);
+	TVertexAttributesConstRef<FVector> ResultVertexPositions = MeshDescription.VertexAttributes().GetAttributesRef<FVector>(MeshAttribute::Vertex::Position);
+	TVertexInstanceAttributesConstRef<FVector> ResultVertexInstanceNormals = MeshDescription.VertexInstanceAttributes().GetAttributesRef<FVector>(MeshAttribute::VertexInstance::Normal);
+	TVertexInstanceAttributesConstRef<FVector> ResultVertexInstanceTangents = MeshDescription.VertexInstanceAttributes().GetAttributesRef<FVector>(MeshAttribute::VertexInstance::Tangent);
+	TVertexInstanceAttributesConstRef<float> ResultVertexInstanceBinormalSigns = MeshDescription.VertexInstanceAttributes().GetAttributesRef<float>(MeshAttribute::VertexInstance::BinormalSign);
+	TVertexInstanceAttributesConstRef<FVector4> ResultVertexInstanceColors = MeshDescription.VertexInstanceAttributes().GetAttributesRef<FVector4>(MeshAttribute::VertexInstance::Color);
+	TVertexInstanceAttributesConstRef<FVector2D> ResultVertexInstanceUVs = MeshDescription.VertexInstanceAttributes().GetAttributesRef<FVector2D>(MeshAttribute::VertexInstance::TextureCoordinate);
+	TEdgeAttributesConstRef<bool> ResultEdgeHardnesses = MeshDescription.EdgeAttributes().GetAttributesRef<bool>(MeshAttribute::Edge::IsHard);
+	TPolygonGroupAttributesConstRef<FName> ResultPolygonGroupMaterialName = MeshDescription.PolygonGroupAttributes().GetAttributesRef<FName>(MeshAttribute::PolygonGroup::ImportedMaterialSlotName);
 
 	
 	//////////////////////////////////////////////////////////////////////////
@@ -432,48 +495,31 @@ bool FMeshDescriptionTest::CompareMeshDescription(const FString& AssetName, FAut
 	FString ConversionName = TEXT("MeshDescription to RawMesh to MeshDescription");
 
 	//Positions
-	MeshDescriptionStructureArrayCompare<TVertexAttributeArray<FVector>, FVertexArray, FVertexID>(ConversionName, AssetName, ExecutionInfo, bAllSame, ReferenceMeshDescription.Vertices(), TEXT("vertex positions"), ReferenceVertexPositions, ResultVertexPositions);
+	MeshDescriptionAttributeArrayCompare(ConversionName, AssetName, ExecutionInfo, bAllSame, ReferenceMeshDescription.Vertices(), TEXT("vertex positions"), ReferenceVertexPositions, ResultVertexPositions);
 
 	//Normals
-	MeshDescriptionStructureArrayCompare<TVertexInstanceAttributeArray<FVector>, FVertexInstanceArray, FVertexInstanceID>(ConversionName, AssetName, ExecutionInfo, bAllSame, ReferenceMeshDescription.VertexInstances(), TEXT("vertex instance normals"), ReferenceVertexInstanceNormals, ResultVertexInstanceNormals);
+	MeshDescriptionAttributeArrayCompare(ConversionName, AssetName, ExecutionInfo, bAllSame, ReferenceMeshDescription.VertexInstances(), TEXT("vertex instance normals"), ReferenceVertexInstanceNormals, ResultVertexInstanceNormals);
 
 	//Tangents
-	MeshDescriptionStructureArrayCompare<TVertexInstanceAttributeArray<FVector>, FVertexInstanceArray, FVertexInstanceID>(ConversionName, AssetName, ExecutionInfo, bAllSame, ReferenceMeshDescription.VertexInstances(), TEXT("vertex instance tangents"), ReferenceVertexInstanceTangents, ResultVertexInstanceTangents);
+	MeshDescriptionAttributeArrayCompare(ConversionName, AssetName, ExecutionInfo, bAllSame, ReferenceMeshDescription.VertexInstances(), TEXT("vertex instance tangents"), ReferenceVertexInstanceTangents, ResultVertexInstanceTangents);
 
 	//BiNormal signs
-	MeshDescriptionNumberArrayCompare<TVertexInstanceAttributeArray<float>, FVertexInstanceArray, FVertexInstanceID>(ConversionName, AssetName, ExecutionInfo, bAllSame, ReferenceMeshDescription.VertexInstances(), TEXT("vertex instance binormals"), ReferenceVertexInstanceBinormalSigns, ResultVertexInstanceBinormalSigns);
+	MeshDescriptionAttributeArrayCompare(ConversionName, AssetName, ExecutionInfo, bAllSame, ReferenceMeshDescription.VertexInstances(), TEXT("vertex instance binormals"), ReferenceVertexInstanceBinormalSigns, ResultVertexInstanceBinormalSigns);
 
 	//Colors
-	MeshDescriptionStructureArrayCompare<TVertexInstanceAttributeArray<FVector4>, FVertexInstanceArray, FVertexInstanceID>(ConversionName, AssetName, ExecutionInfo, bAllSame, ReferenceMeshDescription.VertexInstances(), TEXT("vertex instance colors"), ReferenceVertexInstanceColors, ResultVertexInstanceColors);
+	MeshDescriptionAttributeArrayCompare(ConversionName, AssetName, ExecutionInfo, bAllSame, ReferenceMeshDescription.VertexInstances(), TEXT("vertex instance colors"), ReferenceVertexInstanceColors, ResultVertexInstanceColors);
 
 	//Uvs
-	if (ReferenceVertexInstanceUVs.GetNumIndices() != ResultVertexInstanceUVs.GetNumIndices())
-	{
-		ExecutionInfo.AddEvent(FAutomationEvent(EAutomationEventType::Error, FString::Printf(TEXT("The %s conversion MeshDescription to RawMesh to MeshDescription is not lossless, vertex instance UVs channel count is different. UVs channel count expected [%d] result [%d]"),
-			*AssetName,
-			ReferenceVertexInstanceUVs.GetNumIndices(),
-			ResultVertexInstanceUVs.GetNumIndices())));
-		bAllSame = false;
-	}
-	else
-	{
-		for (int32 UVChannelIndex = 0; UVChannelIndex < ReferenceVertexInstanceUVs.GetNumIndices(); ++UVChannelIndex)
-		{
-			FString UVIndexName = FString::Printf(TEXT("vertex instance UVs(%d)"), UVChannelIndex);
-			const TMeshAttributeArray<FVector2D, FVertexInstanceID>& ReferenceUVs = ReferenceVertexInstanceUVs.GetArrayForIndex(UVChannelIndex);
-			const TMeshAttributeArray<FVector2D, FVertexInstanceID>& ResultUVs = ResultVertexInstanceUVs.GetArrayForIndex(UVChannelIndex);
-			MeshDescriptionStructureArrayCompare<TMeshAttributeArray<FVector2D, FVertexInstanceID>, FVertexInstanceArray, FVertexInstanceID>(ConversionName, AssetName, ExecutionInfo, bAllSame, ReferenceMeshDescription.VertexInstances(), UVIndexName, ReferenceUVs, ResultUVs);
-		}
-	}
+	MeshDescriptionAttributeArrayCompare(ConversionName, AssetName, ExecutionInfo, bAllSame, ReferenceMeshDescription.VertexInstances(), TEXT("vertex instance UVs"), ReferenceVertexInstanceUVs, ResultVertexInstanceUVs);
 
 	//Edges
 	//We do not use a template since we need to check the connected polygon count to validate a false comparison
-	if (ReferenceEdgeHardnesses.Num() != ResultEdgeHardnesses.Num())
+	if (ReferenceEdgeHardnesses.GetNumElements() != ResultEdgeHardnesses.GetNumElements())
 	{
 		ExecutionInfo.AddEvent(FAutomationEvent(EAutomationEventType::Error, FString::Printf(TEXT("The %s conversion MeshDescription to RawMesh to MeshDescription is not lossless, Edge count is different. Edges count expected [%d] result [%d]"),
 			*AssetName,
-			ReferenceEdgeHardnesses.Num(),
-			ResultEdgeHardnesses.Num())));
+			ReferenceEdgeHardnesses.GetNumElements(),
+			ResultEdgeHardnesses.GetNumElements())));
 		bAllSame = false;
 	}
 	else
@@ -500,7 +546,7 @@ bool FMeshDescriptionTest::CompareMeshDescription(const FString& AssetName, FAut
 	}
 
 	//Polygon group ID
-	MeshDescriptionNumberFNameCompare<FPolygonGroupArray, FPolygonGroupID>(ConversionName, AssetName, ExecutionInfo, bAllSame, ReferenceMeshDescription.PolygonGroups(), TEXT("PolygonGroup Material Name"), ReferencePolygonGroupMaterialName, ResultPolygonGroupMaterialName);
+	MeshDescriptionAttributeArrayCompare(ConversionName, AssetName, ExecutionInfo, bAllSame, ReferenceMeshDescription.PolygonGroups(), TEXT("PolygonGroup Material Name"), ReferencePolygonGroupMaterialName, ResultPolygonGroupMaterialName);
 
 	return bAllSame;
 }
@@ -631,11 +677,11 @@ bool FMeshDescriptionTest::NTBTest(FAutomationTestExecutionInfo& ExecutionInfo)
 		AssetMesh->SourceModels[0].LoadRawMesh(RawMesh);
 
 		//const TVertexAttributeArray<FVector>& VertexPositions = MeshDescription.VertexAttributes().GetAttributes<FVector>(MeshAttribute::Vertex::Position);
-		const TVertexInstanceAttributeArray<FVector>& VertexInstanceNormals = MeshDescription.VertexInstanceAttributes().GetAttributes<FVector>(MeshAttribute::VertexInstance::Normal);
-		const TVertexInstanceAttributeArray<FVector>& VertexInstanceTangents = MeshDescription.VertexInstanceAttributes().GetAttributes<FVector>(MeshAttribute::VertexInstance::Tangent);
-		const TVertexInstanceAttributeArray<float>& VertexInstanceBinormalSigns = MeshDescription.VertexInstanceAttributes().GetAttributes<float>(MeshAttribute::VertexInstance::BinormalSign);
-		const TVertexInstanceAttributeArray<FVector4>& VertexInstanceColors = MeshDescription.VertexInstanceAttributes().GetAttributes<FVector4>(MeshAttribute::VertexInstance::Color);
-		const TVertexInstanceAttributeIndicesArray<FVector2D>& VertexInstanceUVs = MeshDescription.VertexInstanceAttributes().GetAttributesSet<FVector2D>(MeshAttribute::VertexInstance::TextureCoordinate);
+		const TVertexInstanceAttributesRef<FVector> VertexInstanceNormals = MeshDescription.VertexInstanceAttributes().GetAttributesRef<FVector>(MeshAttribute::VertexInstance::Normal);
+		const TVertexInstanceAttributesRef<FVector> VertexInstanceTangents = MeshDescription.VertexInstanceAttributes().GetAttributesRef<FVector>(MeshAttribute::VertexInstance::Tangent);
+		const TVertexInstanceAttributesRef<float> VertexInstanceBinormalSigns = MeshDescription.VertexInstanceAttributes().GetAttributesRef<float>(MeshAttribute::VertexInstance::BinormalSign);
+		const TVertexInstanceAttributesRef<FVector4> VertexInstanceColors = MeshDescription.VertexInstanceAttributes().GetAttributesRef<FVector4>(MeshAttribute::VertexInstance::Color);
+		const TVertexInstanceAttributesRef<FVector2D> VertexInstanceUVs = MeshDescription.VertexInstanceAttributes().GetAttributesRef<FVector2D>(MeshAttribute::VertexInstance::TextureCoordinate);
 		int32 ExistingUVCount = VertexInstanceUVs.GetNumIndices();
 		//Build the normals and tangent and compare the result
 		MeshDescription.ComputePolygonTangentsAndNormals(SMALL_NUMBER);
@@ -734,10 +780,10 @@ bool FMeshDescriptionTest::NTBTest(FAutomationTestExecutionInfo& ExecutionInfo)
 
 					for (int32 UVIndex = 0; UVIndex < ExistingUVCount; ++UVIndex)
 					{
-						if (!RawMesh.WedgeTexCoords[UVIndex][WedgeIndex].Equals(VertexInstanceUVs.GetArrayForIndex(UVIndex)[VertexInstanceID], THRESH_UVS_ARE_SAME))
+						if (!RawMesh.WedgeTexCoords[UVIndex][WedgeIndex].Equals(VertexInstanceUVs.Get(VertexInstanceID, UVIndex), THRESH_UVS_ARE_SAME))
 						{
 							OutputError(FString::Printf(TEXT("Vertex Texture coordinnate is different between MeshDescription [%s] and FRawMesh [%s].   UVIndex[%d]  Indice[%d]")
-								, *(VertexInstanceUVs.GetArrayForIndex(UVIndex)[VertexInstanceID].ToString())
+								, *(VertexInstanceUVs.Get(VertexInstanceID, UVIndex).ToString())
 								, *(RawMesh.WedgeTexCoords[UVIndex][WedgeIndex].ToString())
 								, UVIndex
 								, VertexInstanceIDValue));
diff --git a/Engine/Source/Developer/MeshBuilder/Private/MeshDescriptionHelper.cpp b/Engine/Source/Developer/MeshBuilder/Private/MeshDescriptionHelper.cpp
--- a/Engine/Source/Developer/MeshBuilder/Private/MeshDescriptionHelper.cpp
+++ b/Engine/Source/Developer/MeshBuilder/Private/MeshDescriptionHelper.cpp
@@ -45,9 +45,9 @@ void FMeshDescriptionHelper::GetRenderMeshDescription(UObject* Owner, const FMes
 	//OutRenderMeshDescription->ComputePolygonTangentsAndNormals(BuildSettings->bRemoveDegenerates ? SMALL_NUMBER : 0.0f);
 
 	FVertexInstanceArray& VertexInstanceArray = OutRenderMeshDescription.VertexInstances();
-	TVertexInstanceAttributeArray<FVector>& Normals = OutRenderMeshDescription.VertexInstanceAttributes().GetAttributes<FVector>( MeshAttribute::VertexInstance::Normal );
-	TVertexInstanceAttributeArray<FVector>& Tangents = OutRenderMeshDescription.VertexInstanceAttributes().GetAttributes<FVector>( MeshAttribute::VertexInstance::Tangent );
-	TVertexInstanceAttributeArray<float>& BinormalSigns = OutRenderMeshDescription.VertexInstanceAttributes().GetAttributes<float>( MeshAttribute::VertexInstance::BinormalSign );
+	TVertexInstanceAttributesRef<FVector> Normals = OutRenderMeshDescription.VertexInstanceAttributes().GetAttributesRef<FVector>( MeshAttribute::VertexInstance::Normal );
+	TVertexInstanceAttributesRef<FVector> Tangents = OutRenderMeshDescription.VertexInstanceAttributes().GetAttributesRef<FVector>( MeshAttribute::VertexInstance::Tangent );
+	TVertexInstanceAttributesRef<float> BinormalSigns = OutRenderMeshDescription.VertexInstanceAttributes().GetAttributesRef<float>( MeshAttribute::VertexInstance::BinormalSign );
 
 	// Find overlapping corners to accelerate adjacency.
 	FMeshDescriptionOperations::FindOverlappingCorners(OverlappingCorners, OutRenderMeshDescription, ComparisonThreshold);
@@ -108,7 +108,7 @@ void FMeshDescriptionHelper::GetRenderMeshDescription(UObject* Owner, const FMes
 
 	if (BuildSettings->bGenerateLightmapUVs && VertexInstanceArray.Num() > 0)
 	{
-		TVertexInstanceAttributeIndicesArray<FVector2D>& VertexInstanceUVs = OutRenderMeshDescription.VertexInstanceAttributes().GetAttributesSet<FVector2D>(MeshAttribute::VertexInstance::TextureCoordinate);
+		TVertexInstanceAttributesRef<FVector2D> VertexInstanceUVs = OutRenderMeshDescription.VertexInstanceAttributes().GetAttributesRef<FVector2D>(MeshAttribute::VertexInstance::TextureCoordinate);
 		int32 NumIndices = VertexInstanceUVs.GetNumIndices();
 		//Verify the src light map channel
 		if (BuildSettings->SrcLightmapIndex >= NumIndices)
diff --git a/Engine/Source/Developer/MeshBuilder/Private/StaticMeshBuilder.cpp b/Engine/Source/Developer/MeshBuilder/Private/StaticMeshBuilder.cpp
--- a/Engine/Source/Developer/MeshBuilder/Private/StaticMeshBuilder.cpp
+++ b/Engine/Source/Developer/MeshBuilder/Private/StaticMeshBuilder.cpp
@@ -146,7 +146,7 @@ bool FStaticMeshBuilder::Build(FStaticMeshRenderData& StaticMeshRenderData, USta
 			}
 			
 
-			const TPolygonGroupAttributeArray<FName>& PolygonGroupImportedMaterialSlotNames = MeshDescriptions[LodIndex].PolygonGroupAttributes().GetAttributes<FName>(MeshAttribute::PolygonGroup::ImportedMaterialSlotName);
+			const TPolygonGroupAttributesRef<FName> PolygonGroupImportedMaterialSlotNames = MeshDescriptions[LodIndex].PolygonGroupAttributes().GetAttributesRef<FName>(MeshAttribute::PolygonGroup::ImportedMaterialSlotName);
 			// Recompute adjacency information. Since we change the vertices when we reduce
 			MeshDescriptionHelper.FindOverlappingCorners(MeshDescriptions[LodIndex], OverlappingThreshold);
 			
@@ -332,16 +332,15 @@ void BuildVertexBuffer(
 	}
 	TArray<int32> DupVerts;
 
-	const uint32 NumTextureCoord = MeshDescription.VertexInstanceAttributes().GetAttributeIndexCount<FVector2D>(MeshAttribute::VertexInstance::TextureCoordinate);
+	TPolygonGroupAttributesConstRef<FName> PolygonGroupImportedMaterialSlotNames = MeshDescription.PolygonGroupAttributes().GetAttributesRef<FName>(MeshAttribute::PolygonGroup::ImportedMaterialSlotName);
+	TVertexAttributesConstRef<FVector> VertexPositions = MeshDescription.VertexAttributes().GetAttributesRef<FVector>( MeshAttribute::Vertex::Position );
+	TVertexInstanceAttributesConstRef<FVector> VertexInstanceNormals = MeshDescription.VertexInstanceAttributes().GetAttributesRef<FVector>( MeshAttribute::VertexInstance::Normal );
+	TVertexInstanceAttributesConstRef<FVector> VertexInstanceTangents = MeshDescription.VertexInstanceAttributes().GetAttributesRef<FVector>( MeshAttribute::VertexInstance::Tangent );
+	TVertexInstanceAttributesConstRef<float> VertexInstanceBinormalSigns = MeshDescription.VertexInstanceAttributes().GetAttributesRef<float>( MeshAttribute::VertexInstance::BinormalSign );
+	TVertexInstanceAttributesConstRef<FVector4> VertexInstanceColors = MeshDescription.VertexInstanceAttributes().GetAttributesRef<FVector4>( MeshAttribute::VertexInstance::Color );
+	TVertexInstanceAttributesConstRef<FVector2D> VertexInstanceUVs = MeshDescription.VertexInstanceAttributes().GetAttributesRef<FVector2D>( MeshAttribute::VertexInstance::TextureCoordinate );
 
-	const TPolygonGroupAttributeArray<FName>& PolygonGroupImportedMaterialSlotNames = MeshDescription.PolygonGroupAttributes().GetAttributes<FName>(MeshAttribute::PolygonGroup::ImportedMaterialSlotName);
-
-	const TVertexAttributeArray<FVector>& VertexPositions = MeshDescription.VertexAttributes().GetAttributes<FVector>( MeshAttribute::Vertex::Position );
-	const TVertexInstanceAttributeArray<FVector>& VertexInstanceNormals = MeshDescription.VertexInstanceAttributes().GetAttributes<FVector>( MeshAttribute::VertexInstance::Normal );
-	const TVertexInstanceAttributeArray<FVector>& VertexInstanceTangents = MeshDescription.VertexInstanceAttributes().GetAttributes<FVector>( MeshAttribute::VertexInstance::Tangent );
-	const TVertexInstanceAttributeArray<float>& VertexInstanceBinormalSigns = MeshDescription.VertexInstanceAttributes().GetAttributes<float>( MeshAttribute::VertexInstance::BinormalSign );
-	const TVertexInstanceAttributeArray<FVector4>& VertexInstanceColors = MeshDescription.VertexInstanceAttributes().GetAttributes<FVector4>( MeshAttribute::VertexInstance::Color );
-	const TVertexInstanceAttributeIndicesArray<FVector2D>& VertexInstanceUVs = MeshDescription.VertexInstanceAttributes().GetAttributesSet<FVector2D>( MeshAttribute::VertexInstance::TextureCoordinate );
+	const uint32 NumTextureCoord = VertexInstanceUVs.GetNumIndices();
 
 	TMap<FPolygonGroupID, int32> PolygonGroupToSectionIndex;
 	
@@ -425,7 +424,7 @@ void BuildVertexBuffer(
 				{
 					if(UVIndex < NumTextureCoord)
 					{
-						StaticMeshVertex.UVs[UVIndex] = VertexInstanceUVs.GetArrayForIndex(UVIndex)[VertexInstanceID];
+						StaticMeshVertex.UVs[UVIndex] = VertexInstanceUVs.Get(VertexInstanceID, UVIndex);
 					}
 					else
 					{
diff --git a/Engine/Source/Developer/MeshDescriptionOperations/Private/LayoutUV.cpp b/Engine/Source/Developer/MeshDescriptionOperations/Private/LayoutUV.cpp
--- a/Engine/Source/Developer/MeshDescriptionOperations/Private/LayoutUV.cpp
+++ b/Engine/Source/Developer/MeshDescriptionOperations/Private/LayoutUV.cpp
@@ -42,7 +42,7 @@ namespace MeshDescriptionOp
 		int32 WedgeIndex = 0;
 		RemapVerts.SetNumUninitialized(NumIndexes);
 
-		const TVertexInstanceAttributeArray<FVector2D>& VertexUVs = MeshDescription.VertexInstanceAttributes().GetAttributes<FVector2D>(MeshAttribute::VertexInstance::TextureCoordinate, SrcChannel);
+		const TVertexInstanceAttributesRef<FVector2D> VertexUVs = MeshDescription.VertexInstanceAttributes().GetAttributesRef<FVector2D>(MeshAttribute::VertexInstance::TextureCoordinate);
 
 		for (const FPolygonID PolygonID : MeshDescription.Polygons().GetElementIDs())
 		{
@@ -54,7 +54,7 @@ namespace MeshDescriptionOp
 					const FVertexInstanceID VertexInstanceID = MeshTriangle.GetVertexInstanceID(Corner);
 
 					TranslatedMatches[WedgeIndex] = -1;
-					TexCoords[WedgeIndex] = VertexUVs[VertexInstanceID];
+					TexCoords[WedgeIndex] = VertexUVs.Get(VertexInstanceID, SrcChannel);
 					RemapVerts[WedgeIndex] = VertexInstanceID.GetValue();
 					++WedgeIndex;
 				}
@@ -177,7 +177,7 @@ namespace MeshDescriptionOp
 
 		TMap< uint32, int32 > DisjointSetToChartMap;
 
-		const TVertexAttributeArray<FVector>& VertexPositions = MeshDescription.VertexAttributes().GetAttributes<FVector>(MeshAttribute::Vertex::Position);
+		const TVertexAttributesRef<FVector> VertexPositions = MeshDescription.VertexAttributes().GetAttributesRef<FVector>(MeshAttribute::Vertex::Position);
 
 		// Build Charts
 		for (uint32 Tri = 0; Tri < NumTris; )
@@ -1076,7 +1076,7 @@ namespace MeshDescriptionOp
 			ensure(false);	// not expecting it to get here
 		}
 
-		TVertexInstanceAttributeArray<FVector2D>& VertexUVs = MeshDescription.VertexInstanceAttributes().GetAttributes<FVector2D>(MeshAttribute::VertexInstance::TextureCoordinate, DstChannel);
+		TVertexInstanceAttributesRef<FVector2D> VertexUVs = MeshDescription.VertexInstanceAttributes().GetAttributesRef<FVector2D>(MeshAttribute::VertexInstance::TextureCoordinate);
 
 		// Commit chart UVs
 		for (int32 i = 0; i < Charts.Num(); i++)
@@ -1094,7 +1094,7 @@ namespace MeshDescriptionOp
 					uint32 Index = 3 * SortedTris[Tri] + k;
 					const FVector2D& UV = TexCoords[Index];
 					const FVertexInstanceID VertexInstanceID(RemapVerts[Index]);
-					VertexUVs[VertexInstanceID] = UV.X * Chart.PackingScaleU + UV.Y * Chart.PackingScaleV + Chart.PackingBias;
+					VertexUVs.Set(VertexInstanceID, DstChannel, UV.X * Chart.PackingScaleU + UV.Y * Chart.PackingScaleV + Chart.PackingBias);
 				}
 			}
 		}
diff --git a/Engine/Source/Developer/MeshDescriptionOperations/Private/LayoutUV.h b/Engine/Source/Developer/MeshDescriptionOperations/Private/LayoutUV.h
--- a/Engine/Source/Developer/MeshDescriptionOperations/Private/LayoutUV.h
+++ b/Engine/Source/Developer/MeshDescriptionOperations/Private/LayoutUV.h
@@ -99,7 +99,7 @@ namespace MeshDescriptionOp
 		const FVertexID VertexIDA = MeshDescription.GetVertexInstanceVertex(VertexInstanceIDA);
 		const FVertexID VertexIDB = MeshDescription.GetVertexInstanceVertex(VertexInstanceIDB);
 
-		const TVertexAttributeArray<FVector>& VertexPositions = MeshDescription.VertexAttributes().GetAttributes<FVector>(MeshAttribute::Vertex::Position);
+		const TVertexAttributesRef<FVector> VertexPositions = MeshDescription.VertexAttributes().GetAttributesRef<FVector>(MeshAttribute::Vertex::Position);
 		return VertexPositions[VertexIDA].Equals(VertexPositions[VertexIDB], THRESH_POINTS_ARE_SAME);
 	}
 
@@ -117,7 +117,7 @@ namespace MeshDescriptionOp
 		const FVertexInstanceID VertexInstanceIDA(RemapVerts[a]);
 		const FVertexInstanceID VertexInstanceIDB(RemapVerts[b]);
 
-		const TVertexInstanceAttributeArray<FVector>& VertexNormals = MeshDescription.VertexInstanceAttributes().GetAttributes<FVector>(MeshAttribute::VertexInstance::Normal);
+		const TVertexInstanceAttributesRef<FVector> VertexNormals = MeshDescription.VertexInstanceAttributes().GetAttributesRef<FVector>(MeshAttribute::VertexInstance::Normal);
 		return VertexNormals[VertexInstanceIDA].Equals(VertexNormals[VertexInstanceIDB], THRESH_NORMALS_ARE_SAME);
 	}
 
@@ -134,8 +134,8 @@ namespace MeshDescriptionOp
 		const FVertexInstanceID VertexInstanceIDA(RemapVerts[a]);
 		const FVertexInstanceID VertexInstanceIDB(RemapVerts[b]);
 
-		const TVertexInstanceAttributeArray<FVector2D>& VertexUVs = MeshDescription.VertexInstanceAttributes().GetAttributes<FVector2D>(MeshAttribute::VertexInstance::TextureCoordinate, SrcChannel);
-		return VertexUVs[VertexInstanceIDA].Equals(VertexUVs[VertexInstanceIDB], GetUVEqualityThreshold());
+		const TVertexInstanceAttributesRef<FVector2D> VertexUVs = MeshDescription.VertexInstanceAttributes().GetAttributesRef<FVector2D>(MeshAttribute::VertexInstance::TextureCoordinate);
+		return VertexUVs[VertexInstanceIDA].Equals(VertexUVs.Get(VertexInstanceIDB, SrcChannel), GetUVEqualityThreshold());
 	}
 
 	inline bool FLayoutUV::VertsMatch(uint32 a, uint32 b) const
@@ -146,12 +146,12 @@ namespace MeshDescriptionOp
 	// Signed UV area
 	inline float FLayoutUV::TriangleUVArea(uint32 Tri) const
 	{
-		const TVertexInstanceAttributeArray<FVector2D>& VertexUVs = MeshDescription.VertexInstanceAttributes().GetAttributes<FVector2D>(MeshAttribute::VertexInstance::TextureCoordinate, SrcChannel);
+		const TVertexInstanceAttributesRef<FVector2D> VertexUVs = MeshDescription.VertexInstanceAttributes().GetAttributesRef<FVector2D>(MeshAttribute::VertexInstance::TextureCoordinate);
 
 		FVector2D UVs[3];
 		for (int k = 0; k < 3; k++)
 		{
-			UVs[k] = VertexUVs[FVertexInstanceID(RemapVerts[(3 * Tri) + k])];
+			UVs[k] = VertexUVs.Get(FVertexInstanceID(RemapVerts[(3 * Tri) + k]), SrcChannel);
 		}
 
 		FVector2D EdgeUV1 = UVs[1] - UVs[0];
diff --git a/Engine/Source/Developer/MeshDescriptionOperations/Private/MeshDescriptionOperations.cpp b/Engine/Source/Developer/MeshDescriptionOperations/Private/MeshDescriptionOperations.cpp
--- a/Engine/Source/Developer/MeshDescriptionOperations/Private/MeshDescriptionOperations.cpp
+++ b/Engine/Source/Developer/MeshDescriptionOperations/Private/MeshDescriptionOperations.cpp
@@ -71,7 +71,7 @@ void FMeshDescriptionOperations::ConvertHardEdgesToSmoothGroup(const FMeshDescri
 
 	TMap < FPolygonID, uint32> PolygonAvoidances;
 
-	const TEdgeAttributeArray<bool>& EdgeHardnesses = SourceMeshDescription.EdgeAttributes().GetAttributes<bool>(MeshAttribute::Edge::IsHard);
+	TEdgeAttributesConstRef<bool> EdgeHardnesses = SourceMeshDescription.EdgeAttributes().GetAttributesRef<bool>(MeshAttribute::Edge::IsHard);
 
 	for (const FPolygonID PolygonID : SourceMeshDescription.Polygons().GetElementIDs())
 	{
@@ -187,7 +187,7 @@ void FMeshDescriptionOperations::ConvertHardEdgesToSmoothGroup(const FMeshDescri
 
 void FMeshDescriptionOperations::ConvertSmoothGroupToHardEdges(const TArray<uint32>& FaceSmoothingMasks, FMeshDescription& DestinationMeshDescription)
 {
-	TEdgeAttributeArray<bool>& EdgeHardnesses = DestinationMeshDescription.EdgeAttributes().GetAttributes<bool>(MeshAttribute::Edge::IsHard);
+	TEdgeAttributesRef<bool> EdgeHardnesses = DestinationMeshDescription.EdgeAttributes().GetAttributesRef<bool>(MeshAttribute::Edge::IsHard);
 
 	TArray<bool> ConsumedPolygons;
 	ConsumedPolygons.AddZeroed(DestinationMeshDescription.Polygons().Num());
@@ -245,15 +245,15 @@ void FMeshDescriptionOperations::ConvertToRawMesh(const FMeshDescription& Source
 	DestinationRawMesh.Empty();
 
 	//Gather all array data
-	const TVertexAttributeArray<FVector>& VertexPositions = SourceMeshDescription.VertexAttributes().GetAttributes<FVector>(MeshAttribute::Vertex::Position);
+	TVertexAttributesConstRef<FVector> VertexPositions = SourceMeshDescription.VertexAttributes().GetAttributesRef<FVector>(MeshAttribute::Vertex::Position);
 
-	const TVertexInstanceAttributeArray<FVector>& VertexInstanceNormals = SourceMeshDescription.VertexInstanceAttributes().GetAttributes<FVector>(MeshAttribute::VertexInstance::Normal);
-	const TVertexInstanceAttributeArray<FVector>& VertexInstanceTangents = SourceMeshDescription.VertexInstanceAttributes().GetAttributes<FVector>(MeshAttribute::VertexInstance::Tangent);
-	const TVertexInstanceAttributeArray<float>& VertexInstanceBinormalSigns = SourceMeshDescription.VertexInstanceAttributes().GetAttributes<float>(MeshAttribute::VertexInstance::BinormalSign);
-	const TVertexInstanceAttributeArray<FVector4>& VertexInstanceColors = SourceMeshDescription.VertexInstanceAttributes().GetAttributes<FVector4>(MeshAttribute::VertexInstance::Color);
-	const TVertexInstanceAttributeIndicesArray<FVector2D>& VertexInstanceUVs = SourceMeshDescription.VertexInstanceAttributes().GetAttributesSet<FVector2D>(MeshAttribute::VertexInstance::TextureCoordinate);
+	TVertexInstanceAttributesConstRef<FVector> VertexInstanceNormals = SourceMeshDescription.VertexInstanceAttributes().GetAttributesRef<FVector>(MeshAttribute::VertexInstance::Normal);
+	TVertexInstanceAttributesConstRef<FVector> VertexInstanceTangents = SourceMeshDescription.VertexInstanceAttributes().GetAttributesRef<FVector>(MeshAttribute::VertexInstance::Tangent);
+	TVertexInstanceAttributesConstRef<float> VertexInstanceBinormalSigns = SourceMeshDescription.VertexInstanceAttributes().GetAttributesRef<float>(MeshAttribute::VertexInstance::BinormalSign);
+	TVertexInstanceAttributesConstRef<FVector4> VertexInstanceColors = SourceMeshDescription.VertexInstanceAttributes().GetAttributesRef<FVector4>(MeshAttribute::VertexInstance::Color);
+	TVertexInstanceAttributesConstRef<FVector2D> VertexInstanceUVs = SourceMeshDescription.VertexInstanceAttributes().GetAttributesRef<FVector2D>(MeshAttribute::VertexInstance::TextureCoordinate);
 
-	const TPolygonGroupAttributeArray<FName>& PolygonGroupMaterialSlotName = SourceMeshDescription.PolygonGroupAttributes().GetAttributes<FName>(MeshAttribute::PolygonGroup::ImportedMaterialSlotName);
+	TPolygonGroupAttributesConstRef<FName> PolygonGroupMaterialSlotName = SourceMeshDescription.PolygonGroupAttributes().GetAttributesRef<FName>(MeshAttribute::PolygonGroup::ImportedMaterialSlotName);
 
 	DestinationRawMesh.VertexPositions.AddZeroed(SourceMeshDescription.Vertices().Num());
 	TArray<int32> RemapVerts;
@@ -315,7 +315,7 @@ void FMeshDescriptionOperations::ConvertToRawMesh(const FMeshDescription& Source
 				DestinationRawMesh.WedgeTangentZ[WedgeIndex] = VertexInstanceNormals[VertexInstanceID];
 				for (int32 UVIndex = 0; UVIndex < ExistingUVCount; ++UVIndex)
 				{
-					DestinationRawMesh.WedgeTexCoords[UVIndex][WedgeIndex] = VertexInstanceUVs.GetArrayForIndex(UVIndex)[VertexInstanceID];
+					DestinationRawMesh.WedgeTexCoords[UVIndex][WedgeIndex] = VertexInstanceUVs.Get(VertexInstanceID, UVIndex);
 				}
 				++WedgeIndex;
 			}
@@ -330,7 +330,7 @@ void FMeshDescriptionOperations::ConvertToRawMesh(const FMeshDescription& Source
 //We will also weld the vertex position (old FRawMesh is not always welded) and construct a mapping array to match the FVertexID
 void FillMeshDescriptionVertexPositionNoDuplicate(const TArray<FVector>& RawMeshVertexPositions, FMeshDescription& DestinationMeshDescription, TArray<FVertexID>& RemapVertexPosition)
 {
-	TVertexAttributeArray<FVector>& VertexPositions = DestinationMeshDescription.VertexAttributes().GetAttributes<FVector>(MeshAttribute::Vertex::Position);
+	TVertexAttributesRef<FVector> VertexPositions = DestinationMeshDescription.VertexAttributes().GetAttributesRef<FVector>(MeshAttribute::Vertex::Position);
 
 	const int32 NumVertex = RawMeshVertexPositions.Num();
 
@@ -424,13 +424,13 @@ void FMeshDescriptionOperations::ConvertFromRawMesh(const FRawMesh& SourceRawMes
 	DestinationMeshDescription.ReserveNewEdges(SourceRawMesh.WedgeIndices.Num() * 2.5f / 3);
 
 	//Gather all array data
-	TVertexInstanceAttributeArray<FVector>& VertexInstanceNormals = DestinationMeshDescription.VertexInstanceAttributes().GetAttributes<FVector>(MeshAttribute::VertexInstance::Normal);
-	TVertexInstanceAttributeArray<FVector>& VertexInstanceTangents = DestinationMeshDescription.VertexInstanceAttributes().GetAttributes<FVector>(MeshAttribute::VertexInstance::Tangent);
-	TVertexInstanceAttributeArray<float>& VertexInstanceBinormalSigns = DestinationMeshDescription.VertexInstanceAttributes().GetAttributes<float>(MeshAttribute::VertexInstance::BinormalSign);
-	TVertexInstanceAttributeArray<FVector4>& VertexInstanceColors = DestinationMeshDescription.VertexInstanceAttributes().GetAttributes<FVector4>(MeshAttribute::VertexInstance::Color);
-	TVertexInstanceAttributeIndicesArray<FVector2D>& VertexInstanceUVs = DestinationMeshDescription.VertexInstanceAttributes().GetAttributesSet<FVector2D>(MeshAttribute::VertexInstance::TextureCoordinate);
+	TVertexInstanceAttributesRef<FVector> VertexInstanceNormals = DestinationMeshDescription.VertexInstanceAttributes().GetAttributesRef<FVector>(MeshAttribute::VertexInstance::Normal);
+	TVertexInstanceAttributesRef<FVector> VertexInstanceTangents = DestinationMeshDescription.VertexInstanceAttributes().GetAttributesRef<FVector>(MeshAttribute::VertexInstance::Tangent);
+	TVertexInstanceAttributesRef<float> VertexInstanceBinormalSigns = DestinationMeshDescription.VertexInstanceAttributes().GetAttributesRef<float>(MeshAttribute::VertexInstance::BinormalSign);
+	TVertexInstanceAttributesRef<FVector4> VertexInstanceColors = DestinationMeshDescription.VertexInstanceAttributes().GetAttributesRef<FVector4>(MeshAttribute::VertexInstance::Color);
+	TVertexInstanceAttributesRef<FVector2D> VertexInstanceUVs = DestinationMeshDescription.VertexInstanceAttributes().GetAttributesRef<FVector2D>(MeshAttribute::VertexInstance::TextureCoordinate);
 
-	TPolygonGroupAttributeArray<FName>& PolygonGroupImportedMaterialSlotNames = DestinationMeshDescription.PolygonGroupAttributes().GetAttributes<FName>(MeshAttribute::PolygonGroup::ImportedMaterialSlotName);
+	TPolygonGroupAttributesRef<FName> PolygonGroupImportedMaterialSlotNames = DestinationMeshDescription.PolygonGroupAttributes().GetAttributesRef<FName>(MeshAttribute::PolygonGroup::ImportedMaterialSlotName);
 
 	int32 NumTexCoords = 0;
 	int32 MaxTexCoords = MAX_MESH_TEXTURE_COORDS;
@@ -532,8 +532,7 @@ void FMeshDescriptionOperations::ConvertFromRawMesh(const FRawMesh& SourceRawMes
 				{
 					continue;
 				}
-				TMeshAttributeArray<FVector2D, FVertexInstanceID>& Uvs = VertexInstanceUVs.GetArrayForIndex(TextureCoordIndex);
-				Uvs[VertexInstanceID] = SourceRawMesh.WedgeTexCoords[TextureCoordinnateIndex][VerticeIndex];
+				VertexInstanceUVs.Set(VertexInstanceID, TextureCoordIndex, SourceRawMesh.WedgeTexCoords[TextureCoordinnateIndex][VerticeIndex]);
 			}
 		}
 
@@ -595,11 +594,11 @@ void FMeshDescriptionOperations::ConvertFromRawMesh(const FRawMesh& SourceRawMes
 
 void FMeshDescriptionOperations::CreatePolygonNTB(FMeshDescription& MeshDescription, float ComparisonThreshold)
 {
-	const TVertexAttributeArray<FVector>& VertexPositions = MeshDescription.VertexAttributes().GetAttributes<FVector>(MeshAttribute::Vertex::Position);
-	TVertexInstanceAttributeArray<FVector2D>& VertexUVs = MeshDescription.VertexInstanceAttributes().GetAttributes<FVector2D>(MeshAttribute::VertexInstance::TextureCoordinate, 0);
-	TPolygonAttributeArray<FVector>& PolygonNormals = MeshDescription.PolygonAttributes().GetAttributes<FVector>(MeshAttribute::Polygon::Normal);
-	TPolygonAttributeArray<FVector>& PolygonTangents = MeshDescription.PolygonAttributes().GetAttributes<FVector>(MeshAttribute::Polygon::Tangent);
-	TPolygonAttributeArray<FVector>& PolygonBinormals = MeshDescription.PolygonAttributes().GetAttributes<FVector>(MeshAttribute::Polygon::Binormal);
+	const TVertexAttributesRef<FVector> VertexPositions = MeshDescription.VertexAttributes().GetAttributesRef<FVector>(MeshAttribute::Vertex::Position);
+	TVertexInstanceAttributesRef<FVector2D> VertexUVs = MeshDescription.VertexInstanceAttributes().GetAttributesRef<FVector2D>(MeshAttribute::VertexInstance::TextureCoordinate);
+	TPolygonAttributesRef<FVector> PolygonNormals = MeshDescription.PolygonAttributes().GetAttributesRef<FVector>(MeshAttribute::Polygon::Normal);
+	TPolygonAttributesRef<FVector> PolygonTangents = MeshDescription.PolygonAttributes().GetAttributesRef<FVector>(MeshAttribute::Polygon::Tangent);
+	TPolygonAttributesRef<FVector> PolygonBinormals = MeshDescription.PolygonAttributes().GetAttributesRef<FVector>(MeshAttribute::Polygon::Binormal);
 
 	FVertexInstanceArray& VertexInstanceArray = MeshDescription.VertexInstances();
 	FVertexArray& VertexArray = MeshDescription.Vertices();
@@ -626,7 +625,7 @@ void FMeshDescriptionOperations::CreatePolygonNTB(FMeshDescription& MeshDescript
 			for (int32 i = 0; i < 3; ++i)
 			{
 				const FVertexInstanceID VertexInstanceID = MeshTriangle.GetVertexInstanceID(i);
-				UVs[i] = VertexUVs[VertexInstanceID];
+				UVs[i] = VertexUVs.Get(VertexInstanceID, 0); // UV0
 				P[i] = VertexPositions[MeshDescription.GetVertexInstanceVertex(VertexInstanceID)];
 			}
 
@@ -700,14 +699,14 @@ void FMeshDescriptionOperations::CreateNormals(FMeshDescription& MeshDescription
 // the angle it makes with the vertex being calculated. This means that triangulated faces whose
 // internal edge meets the vertex doesn't get undue extra weight.
 
-	const TVertexInstanceAttributeArray<FVector2D>& VertexUVs = MeshDescription.VertexInstanceAttributes().GetAttributes<FVector2D>(MeshAttribute::VertexInstance::TextureCoordinate, 0);
-	TVertexInstanceAttributeArray<FVector>& VertexNormals = MeshDescription.VertexInstanceAttributes().GetAttributes<FVector>(MeshAttribute::VertexInstance::Normal, 0);
-	TVertexInstanceAttributeArray<FVector>& VertexTangents = MeshDescription.VertexInstanceAttributes().GetAttributes<FVector>(MeshAttribute::VertexInstance::Tangent, 0);
-	TVertexInstanceAttributeArray<float>& VertexBinormalSigns = MeshDescription.VertexInstanceAttributes().GetAttributes<float>(MeshAttribute::VertexInstance::BinormalSign, 0);
+	const TVertexInstanceAttributesRef<FVector2D> VertexUVs = MeshDescription.VertexInstanceAttributes().GetAttributesRef<FVector2D>(MeshAttribute::VertexInstance::TextureCoordinate);
+	TVertexInstanceAttributesRef<FVector> VertexNormals = MeshDescription.VertexInstanceAttributes().GetAttributesRef<FVector>(MeshAttribute::VertexInstance::Normal);
+	TVertexInstanceAttributesRef<FVector> VertexTangents = MeshDescription.VertexInstanceAttributes().GetAttributesRef<FVector>(MeshAttribute::VertexInstance::Tangent);
+	TVertexInstanceAttributesRef<float> VertexBinormalSigns = MeshDescription.VertexInstanceAttributes().GetAttributesRef<float>(MeshAttribute::VertexInstance::BinormalSign);
 
-	TPolygonAttributeArray<FVector>& PolygonNormals = MeshDescription.PolygonAttributes().GetAttributes<FVector>(MeshAttribute::Polygon::Normal);
-	TPolygonAttributeArray<FVector>& PolygonTangents = MeshDescription.PolygonAttributes().GetAttributes<FVector>(MeshAttribute::Polygon::Tangent);
-	TPolygonAttributeArray<FVector>& PolygonBinormals = MeshDescription.PolygonAttributes().GetAttributes<FVector>(MeshAttribute::Polygon::Binormal);
+	TPolygonAttributesRef<FVector> PolygonNormals = MeshDescription.PolygonAttributes().GetAttributesRef<FVector>(MeshAttribute::Polygon::Normal);
+	TPolygonAttributesRef<FVector> PolygonTangents = MeshDescription.PolygonAttributes().GetAttributesRef<FVector>(MeshAttribute::Polygon::Tangent);
+	TPolygonAttributesRef<FVector> PolygonBinormals = MeshDescription.PolygonAttributes().GetAttributesRef<FVector>(MeshAttribute::Polygon::Binormal);
 
 	TMap<FPolygonID, FVertexInfo> VertexInfoMap;
 	VertexInfoMap.Reserve(20);
@@ -732,7 +731,7 @@ void FMeshDescriptionOperations::CreateNormals(FMeshDescription& MeshDescription
 						if (MeshDescription.GetVertexInstanceVertex(VertexInstanceID) == VertexID)
 						{
 							VertexInfo.VertexInstanceID = VertexInstanceID;
-							VertexInfo.UVs = VertexUVs[VertexInstanceID];
+							VertexInfo.UVs = VertexUVs.Get(VertexInstanceID, 0);	// UV0
 							bPointHasAllTangents &= !VertexNormals[VertexInstanceID].IsNearlyZero() && !VertexTangents[VertexInstanceID].IsNearlyZero();
 							if (bPointHasAllTangents)
 							{
@@ -778,7 +777,7 @@ void FMeshDescriptionOperations::CreateNormals(FMeshDescription& MeshDescription
 				FVertexInfo& CurrentVertexInfo = VertexInfoMap.FindOrAdd(CurrentPolygonID);
 				CurrentGroup.AddUnique(CurrentVertexInfo.PolygonID);
 				ConsumedPolygon.AddUnique(CurrentVertexInfo.PolygonID);
-				const TEdgeAttributeArray<bool>& EdgeHardnesses = MeshDescription.EdgeAttributes().GetAttributes<bool>(MeshAttribute::Edge::IsHard);
+				const TEdgeAttributesRef<bool> EdgeHardnesses = MeshDescription.EdgeAttributes().GetAttributesRef<bool>(MeshAttribute::Edge::IsHard);
 				for (const FEdgeID EdgeID : CurrentVertexInfo.EdgeIDs)
 				{
 					if (EdgeHardnesses[EdgeID])
@@ -867,7 +866,7 @@ void FMeshDescriptionOperations::CreateNormals(FMeshDescription& MeshDescription
 			//Apply the average NTB on all Vertex instance
 			for (const FVertexInstanceID VertexInstanceID : VertexInstanceInGroup)
 			{
-				const FVector2D VertexUV = VertexUVs[VertexInstanceID];
+				const FVector2D VertexUV = VertexUVs.Get(VertexInstanceID, 0);	// UV0
 
 				if (VertexNormals[VertexInstanceID].IsNearlyZero(SMALL_NUMBER))
 				{
@@ -1018,7 +1017,7 @@ void FMeshDescriptionOperations::FindOverlappingCorners(TMultiMap<int32, int32>&
 	TArray<MeshDescriptionOperationNamespace::FIndexAndZ> VertIndexAndZ;
 	VertIndexAndZ.Reserve(NumWedges);
 
-	const TVertexAttributeArray<FVector>& VertexPositions = MeshDescription.VertexAttributes().GetAttributes<FVector>(MeshAttribute::Vertex::Position);
+	TVertexAttributesConstRef<FVector> VertexPositions = MeshDescription.VertexAttributes().GetAttributesRef<FVector>(MeshAttribute::Vertex::Position);
 
 	for (const FVertexInstanceID VertexInstanceID : VertexInstanceArray.GetElementIDs())
 	{
@@ -1084,15 +1083,14 @@ bool FMeshDescriptionOperations::GenerateUniqueUVsForStaticMesh(const FMeshDescr
 	if (bPackSuccess)
 	{
 		Packer.CommitPackedUVs();
-		TVertexInstanceAttributeIndicesArray<FVector2D>& VertexInstanceUVs = DuplicateMeshDescription.VertexInstanceAttributes().GetAttributesSet<FVector2D>(MeshAttribute::VertexInstance::TextureCoordinate);
+		TVertexInstanceAttributesRef<FVector2D> VertexInstanceUVs = DuplicateMeshDescription.VertexInstanceAttributes().GetAttributesRef<FVector2D>(MeshAttribute::VertexInstance::TextureCoordinate);
 		// Save generated UVs
 		check(VertexInstanceUVs.GetNumIndices() > 1);
-		auto& UniqueUVsArray = VertexInstanceUVs.GetArrayForIndex(1);
-		OutTexCoords.AddZeroed(UniqueUVsArray.Num());
+		OutTexCoords.AddZeroed(VertexInstanceUVs.GetNumElements());
 		int32 TextureCoordIndex = 0;
 		for (const FVertexInstanceID& VertexInstanceID : DuplicateMeshDescription.VertexInstances().GetElementIDs())
 		{
-			OutTexCoords[TextureCoordIndex++] = UniqueUVsArray[VertexInstanceID];
+			OutTexCoords[TextureCoordIndex++] = VertexInstanceUVs.Get(VertexInstanceID, 1);	// UV1
 		}
 	}
 
diff --git a/Engine/Source/Developer/MeshSimplifier/Private/QuadricMeshReduction.cpp b/Engine/Source/Developer/MeshSimplifier/Private/QuadricMeshReduction.cpp
--- a/Engine/Source/Developer/MeshSimplifier/Private/QuadricMeshReduction.cpp
+++ b/Engine/Source/Developer/MeshSimplifier/Private/QuadricMeshReduction.cpp
@@ -533,15 +533,15 @@ public:
 		}
 		int32 NumWedges = NumFaces * 3;
 
-		const TVertexAttributeArray<FVector>& InVertexPositions = InMesh.VertexAttributes().GetAttributes<FVector>(MeshAttribute::Vertex::Position);
-		const TVertexInstanceAttributeArray<FVector>& InVertexNormals = InMesh.VertexInstanceAttributes().GetAttributes<FVector>(MeshAttribute::VertexInstance::Normal);
-		const TVertexInstanceAttributeArray<FVector>& InVertexTangents = InMesh.VertexInstanceAttributes().GetAttributes<FVector>(MeshAttribute::VertexInstance::Tangent);
-		const TVertexInstanceAttributeArray<float>& InVertexBinormalSigns = InMesh.VertexInstanceAttributes().GetAttributes<float>(MeshAttribute::VertexInstance::BinormalSign);
-		const TVertexInstanceAttributeArray<FVector4>& InVertexColors = InMesh.VertexInstanceAttributes().GetAttributes<FVector4>(MeshAttribute::VertexInstance::Color);
-		const TVertexInstanceAttributeIndicesArray<FVector2D>& InVertexUVs = InMesh.VertexInstanceAttributes().GetAttributesSet<FVector2D>(MeshAttribute::VertexInstance::TextureCoordinate);
-		
-		const TPolygonGroupAttributeArray<FName>& InPolygonGroupMaterialNames = InMesh.PolygonGroupAttributes().GetAttributes<FName>(MeshAttribute::PolygonGroup::ImportedMaterialSlotName);
-		TPolygonGroupAttributeArray<FName>& OutPolygonGroupMaterialNames = OutReducedMesh.PolygonGroupAttributes().GetAttributes<FName>(MeshAttribute::PolygonGroup::ImportedMaterialSlotName);
+		TVertexAttributesConstRef<FVector> InVertexPositions = InMesh.VertexAttributes().GetAttributesRef<FVector>(MeshAttribute::Vertex::Position);
+		TVertexInstanceAttributesConstRef<FVector> InVertexNormals = InMesh.VertexInstanceAttributes().GetAttributesRef<FVector>(MeshAttribute::VertexInstance::Normal);
+		TVertexInstanceAttributesConstRef<FVector> InVertexTangents = InMesh.VertexInstanceAttributes().GetAttributesRef<FVector>(MeshAttribute::VertexInstance::Tangent);
+		TVertexInstanceAttributesConstRef<float> InVertexBinormalSigns = InMesh.VertexInstanceAttributes().GetAttributesRef<float>(MeshAttribute::VertexInstance::BinormalSign);
+		TVertexInstanceAttributesConstRef<FVector4> InVertexColors = InMesh.VertexInstanceAttributes().GetAttributesRef<FVector4>(MeshAttribute::VertexInstance::Color);
+		TVertexInstanceAttributesConstRef<FVector2D> InVertexUVs = InMesh.VertexInstanceAttributes().GetAttributesRef<FVector2D>(MeshAttribute::VertexInstance::TextureCoordinate);
+		TPolygonGroupAttributesConstRef<FName> InPolygonGroupMaterialNames = InMesh.PolygonGroupAttributes().GetAttributesRef<FName>(MeshAttribute::PolygonGroup::ImportedMaterialSlotName);
+
+		TPolygonGroupAttributesRef<FName> OutPolygonGroupMaterialNames = OutReducedMesh.PolygonGroupAttributes().GetAttributesRef<FName>(MeshAttribute::PolygonGroup::ImportedMaterialSlotName);
 
 		int32 FaceIndex = 0;
 		for (const FPolygonID& PolygonID : InMesh.Polygons().GetElementIDs())
@@ -606,7 +606,7 @@ public:
 					{
 						if (UVIndex < InVertexUVs.GetNumIndices())
 						{
-							NewVert.TexCoords[UVIndex] = InVertexUVs.GetArrayForIndex(UVIndex)[VertexInstanceIDs[CornerIndex]];
+							NewVert.TexCoords[UVIndex] = InVertexUVs.Get(VertexInstanceIDs[CornerIndex], UVIndex);
 							InMeshNumTexCoords = FMath::Max(UVIndex+1, InMeshNumTexCoords);
 						}
 						else
@@ -735,7 +735,7 @@ public:
 				OutPolygonGroupMaterialNames[PolygonGroupID] = InPolygonGroupMaterialNames[PolygonGroupID];
 			}
 
-			TVertexAttributeArray<FVector>& OutVertexPositions = OutReducedMesh.VertexAttributes().GetAttributes<FVector>(MeshAttribute::Vertex::Position);
+			TVertexAttributesRef<FVector> OutVertexPositions = OutReducedMesh.VertexAttributes().GetAttributesRef<FVector>(MeshAttribute::Vertex::Position);
 
 			//Fill the vertex array
 			for (int32 VertexIndex = 0; VertexIndex < (int32)NumVerts; ++VertexIndex)
@@ -747,11 +747,11 @@ public:
 
 			TMap<int32, FPolygonGroupID> PolygonGroupMapping;
 
-			TVertexInstanceAttributeArray<FVector>& OutVertexNormals = OutReducedMesh.VertexInstanceAttributes().GetAttributes<FVector>(MeshAttribute::VertexInstance::Normal);
-			TVertexInstanceAttributeArray<FVector>& OutVertexTangents = OutReducedMesh.VertexInstanceAttributes().GetAttributes<FVector>(MeshAttribute::VertexInstance::Tangent);
-			TVertexInstanceAttributeArray<float>& OutVertexBinormalSigns = OutReducedMesh.VertexInstanceAttributes().GetAttributes<float>(MeshAttribute::VertexInstance::BinormalSign);
-			TVertexInstanceAttributeArray<FVector4>& OutVertexColors = OutReducedMesh.VertexInstanceAttributes().GetAttributes<FVector4>(MeshAttribute::VertexInstance::Color);
-			TVertexInstanceAttributeIndicesArray<FVector2D>& OutVertexUVs = OutReducedMesh.VertexInstanceAttributes().GetAttributesSet<FVector2D>(MeshAttribute::VertexInstance::TextureCoordinate);
+			TVertexInstanceAttributesRef<FVector> OutVertexNormals = OutReducedMesh.VertexInstanceAttributes().GetAttributesRef<FVector>(MeshAttribute::VertexInstance::Normal);
+			TVertexInstanceAttributesRef<FVector> OutVertexTangents = OutReducedMesh.VertexInstanceAttributes().GetAttributesRef<FVector>(MeshAttribute::VertexInstance::Tangent);
+			TVertexInstanceAttributesRef<float> OutVertexBinormalSigns = OutReducedMesh.VertexInstanceAttributes().GetAttributesRef<float>(MeshAttribute::VertexInstance::BinormalSign);
+			TVertexInstanceAttributesRef<FVector4> OutVertexColors = OutReducedMesh.VertexInstanceAttributes().GetAttributesRef<FVector4>(MeshAttribute::VertexInstance::Color);
+			TVertexInstanceAttributesRef<FVector2D> OutVertexUVs = OutReducedMesh.VertexInstanceAttributes().GetAttributesRef<FVector2D>(MeshAttribute::VertexInstance::TextureCoordinate);
 
 			//Specify the number of texture coords in this mesh description
 			OutVertexUVs.SetNumIndices(InMeshNumTexCoords);
@@ -786,7 +786,7 @@ public:
 					//Texture coord
 					for (int32 TexCoordIndex = 0; TexCoordIndex < InMeshNumTexCoords; TexCoordIndex++)
 					{
-						OutVertexUVs.GetArrayForIndex(TexCoordIndex)[AddedVertexInstanceId] = Verts[Indexes[VertexInstanceIndex]].TexCoords[TexCoordIndex];
+						OutVertexUVs.Set(AddedVertexInstanceId, TexCoordIndex, Verts[Indexes[VertexInstanceIndex]].TexCoords[TexCoordIndex]);
 					}
 				}
 				
@@ -800,15 +800,14 @@ public:
 					MaterialPolygonGroupID = OutReducedMesh.PolygonGroups().Num() > MaterialIndex ? PolygonGroupID : OutReducedMesh.CreatePolygonGroup();
 
 					// Copy all attributes from the base polygon group to the new polygon group
-					InMesh.PolygonGroupAttributes().ForEachAttributeIndicesArray(
-						[&OutReducedMesh, PolygonGroupID, MaterialPolygonGroupID](const FName Name, const auto& AttributeIndicesArray)
+					InMesh.PolygonGroupAttributes().ForEach(
+						[&OutReducedMesh, PolygonGroupID, MaterialPolygonGroupID](const FName Name, const auto ArrayRef)
 						{
-							for (int32 Index = 0; Index < AttributeIndicesArray.GetNumIndices(); ++Index)
+							for (int32 Index = 0; Index < ArrayRef.GetNumIndices(); ++Index)
 							{
 								// Only copy shared attribute values, since input mesh description can differ from output mesh description
-								const auto& Value = AttributeIndicesArray.GetArrayForIndex(Index)[PolygonGroupID];
-								using AttributeType = typename TDecay<decltype( Value )>::Type;
-								if (OutReducedMesh.PolygonGroupAttributes().HasAttribute<AttributeType>(Name))
+								const auto& Value = ArrayRef.Get(PolygonGroupID, Index);
+								if (OutReducedMesh.PolygonGroupAttributes().HasAttribute(Name))
 								{
 									OutReducedMesh.PolygonGroupAttributes().SetAttribute(MaterialPolygonGroupID, Name, Index, Value);
 								}
diff --git a/Engine/Source/Editor/UnrealEd/Private/Fbx/FbxStaticMeshImport.cpp b/Engine/Source/Editor/UnrealEd/Private/Fbx/FbxStaticMeshImport.cpp
--- a/Engine/Source/Editor/UnrealEd/Private/Fbx/FbxStaticMeshImport.cpp
+++ b/Engine/Source/Editor/UnrealEd/Private/Fbx/FbxStaticMeshImport.cpp
@@ -527,18 +527,16 @@ bool UnFbx::FFbxImporter::BuildStaticMeshFromGeometry(FbxNode* Node, UStaticMesh
 	int32 VertexCount = Mesh->GetControlPointsCount();
 	bool OddNegativeScale = IsOddNegativeScale(TotalMatrix);
 	bool bHasNonDegeneratePolygons = false;
-	TVertexAttributeArray<FVector>& VertexPositions = MeshDescription->VertexAttributes().GetAttributes<FVector>(MeshAttribute::Vertex::Position);
 
-	TVertexInstanceAttributeArray<FVector>& VertexInstanceNormals = MeshDescription->VertexInstanceAttributes().GetAttributes<FVector>(MeshAttribute::VertexInstance::Normal);
-	TVertexInstanceAttributeArray<FVector>& VertexInstanceTangents = MeshDescription->VertexInstanceAttributes().GetAttributes<FVector>(MeshAttribute::VertexInstance::Tangent);
-	TVertexInstanceAttributeArray<float>& VertexInstanceBinormalSigns = MeshDescription->VertexInstanceAttributes().GetAttributes<float>(MeshAttribute::VertexInstance::BinormalSign);
-	TVertexInstanceAttributeArray<FVector4>& VertexInstanceColors = MeshDescription->VertexInstanceAttributes().GetAttributes<FVector4>(MeshAttribute::VertexInstance::Color);
-	TVertexInstanceAttributeIndicesArray<FVector2D>& VertexInstanceUVs = MeshDescription->VertexInstanceAttributes().GetAttributesSet<FVector2D>(MeshAttribute::VertexInstance::TextureCoordinate);
-
-	TEdgeAttributeArray<bool>& EdgeHardnesses = MeshDescription->EdgeAttributes().GetAttributes<bool>(MeshAttribute::Edge::IsHard);
-	TEdgeAttributeArray<float>& EdgeCreaseSharpnesses = MeshDescription->EdgeAttributes().GetAttributes<float>(MeshAttribute::Edge::CreaseSharpness);
-
-	TPolygonGroupAttributeArray<FName>& PolygonGroupImportedMaterialSlotNames = MeshDescription->PolygonGroupAttributes().GetAttributes<FName>(MeshAttribute::PolygonGroup::ImportedMaterialSlotName);
+	TVertexAttributesRef<FVector> VertexPositions = MeshDescription->VertexAttributes().GetAttributesRef<FVector>(MeshAttribute::Vertex::Position);
+	TVertexInstanceAttributesRef<FVector> VertexInstanceNormals = MeshDescription->VertexInstanceAttributes().GetAttributesRef<FVector>(MeshAttribute::VertexInstance::Normal);
+	TVertexInstanceAttributesRef<FVector> VertexInstanceTangents = MeshDescription->VertexInstanceAttributes().GetAttributesRef<FVector>(MeshAttribute::VertexInstance::Tangent);
+	TVertexInstanceAttributesRef<float> VertexInstanceBinormalSigns = MeshDescription->VertexInstanceAttributes().GetAttributesRef<float>(MeshAttribute::VertexInstance::BinormalSign);
+	TVertexInstanceAttributesRef<FVector4> VertexInstanceColors = MeshDescription->VertexInstanceAttributes().GetAttributesRef<FVector4>(MeshAttribute::VertexInstance::Color);
+	TVertexInstanceAttributesRef<FVector2D> VertexInstanceUVs = MeshDescription->VertexInstanceAttributes().GetAttributesRef<FVector2D>(MeshAttribute::VertexInstance::TextureCoordinate);
+	TEdgeAttributesRef<bool> EdgeHardnesses = MeshDescription->EdgeAttributes().GetAttributesRef<bool>(MeshAttribute::Edge::IsHard);
+	TEdgeAttributesRef<float> EdgeCreaseSharpnesses = MeshDescription->EdgeAttributes().GetAttributesRef<float>(MeshAttribute::Edge::CreaseSharpness);
+	TPolygonGroupAttributesRef<FName> PolygonGroupImportedMaterialSlotNames = MeshDescription->PolygonGroupAttributes().GetAttributesRef<FName>(MeshAttribute::PolygonGroup::ImportedMaterialSlotName);
 
 	int32 VertexOffset = MeshDescription->Vertices().Num();
 	int32 VertexInstanceOffset = MeshDescription->VertexInstances().Num();
@@ -547,18 +545,7 @@ bool UnFbx::FFbxImporter::BuildStaticMeshFromGeometry(FbxNode* Node, UStaticMesh
 	TMap<int32, FPolygonGroupID> PolygonGroupMapping;
 
 	// When importing multiple mesh pieces to the same static mesh.  Ensure each mesh piece has the same number of Uv's
-	int32 ExistingUVCount = 0;
-	for (int32 UVChannelIndex = 0; UVChannelIndex < VertexInstanceUVs.GetNumIndices(); ++UVChannelIndex)
-	{
-		if (VertexInstanceUVs.GetArrayForIndex(UVChannelIndex).Num() > 0)
-		{
-			ExistingUVCount++;
-		}
-		else
-		{
-			break;
-		}
-	}
+	int32 ExistingUVCount = VertexInstanceUVs.GetNumIndices();
 
 	int32 NumUVs = FMath::Max(FBXUVs.UniqueUVCount, ExistingUVCount);
 	NumUVs = FMath::Min<int32>(MAX_MESH_TEXTURE_COORDS, NumUVs);
@@ -678,7 +665,7 @@ bool UnFbx::FFbxImporter::BuildStaticMeshFromGeometry(FbxNode* Node, UStaticMesh
 					FinalUVVector.X = static_cast<float>(UVVector[0]);
 					FinalUVVector.Y = 1.f - static_cast<float>(UVVector[1]);   //flip the Y of UVs for DirectX
 				}
-				VertexInstanceUVs.GetArrayForIndex(UVLayerIndex)[AddedVertexInstanceId] = FinalUVVector;
+				VertexInstanceUVs.Set(AddedVertexInstanceId, UVLayerIndex, FinalUVVector);
 			}
 
 			//Color attribute
@@ -1638,9 +1625,9 @@ UStaticMesh* UnFbx::FFbxImporter::ImportStaticMeshAsSingle(UObject* InParent, TA
 
 	if (bBuildStatus)
 	{
-		TVertexInstanceAttributeIndicesArray<FVector2D>& VertexInstanceUVs = MeshDescription->VertexInstanceAttributes().GetAttributesSet<FVector2D>(MeshAttribute::VertexInstance::TextureCoordinate);
+		TVertexInstanceAttributesRef<FVector2D> VertexInstanceUVs = MeshDescription->VertexInstanceAttributes().GetAttributesRef<FVector2D>(MeshAttribute::VertexInstance::TextureCoordinate);
 		int32 FirstOpenUVChannel = VertexInstanceUVs.GetNumIndices() >= MAX_MESH_TEXTURE_COORDS ? 1 : VertexInstanceUVs.GetNumIndices();
-		TPolygonGroupAttributeArray<FName>& PolygonGroupImportedMaterialSlotNames = MeshDescription->PolygonGroupAttributes().GetAttributes<FName>(MeshAttribute::PolygonGroup::ImportedMaterialSlotName);
+		TPolygonGroupAttributesRef<FName> PolygonGroupImportedMaterialSlotNames = MeshDescription->PolygonGroupAttributes().GetAttributesRef<FName>(MeshAttribute::PolygonGroup::ImportedMaterialSlotName);
 
 		TArray<FStaticMaterial> MaterialToAdd;
 		for (const FPolygonGroupID PolygonGroupID : MeshDescription->PolygonGroups().GetElementIDs())
diff --git a/Engine/Source/Editor/UnrealEd/Private/GeomFitUtils.cpp b/Engine/Source/Editor/UnrealEd/Private/GeomFitUtils.cpp
--- a/Engine/Source/Editor/UnrealEd/Private/GeomFitUtils.cpp
+++ b/Engine/Source/Editor/UnrealEd/Private/GeomFitUtils.cpp
@@ -204,7 +204,7 @@ static void CalcBoundingBox(const FRawMesh& RawMesh, FVector& Center, FVector& E
 
 static void CalcBoundingBox(const FMeshDescription& MeshDescription, FVector& Center, FVector& Extents, FVector& LimitVec)
 {
-	const TVertexAttributeArray<FVector>& VertexPositions = MeshDescription.VertexAttributes().GetAttributes<FVector>(MeshAttribute::Vertex::Position);
+	TVertexAttributesConstRef<FVector> VertexPositions = MeshDescription.VertexAttributes().GetAttributesRef<FVector>(MeshAttribute::Vertex::Position);
 
 	FBox Box(ForceInit);
 	for (const FVertexID VertexID : MeshDescription.Vertices().GetElementIDs())
@@ -400,7 +400,7 @@ static void CalcBoundingSphere(const FMeshDescription& MeshDescription, FSphere&
 	FVector MinIx[3];
 	FVector MaxIx[3];
 
-	const TVertexAttributeArray<FVector>& VertexPositions = MeshDescription.VertexAttributes().GetAttributes<FVector>(MeshAttribute::Vertex::Position);
+	TVertexAttributesConstRef<FVector> VertexPositions = MeshDescription.VertexAttributes().GetAttributesRef<FVector>(MeshAttribute::Vertex::Position);
 
 	bool bFirstVertex = true;
 	for (const FVertexID VertexID : MeshDescription.Vertices().GetElementIDs())
@@ -515,7 +515,7 @@ static void CalcBoundingSphere2(const FMeshDescription& MeshDescription, FSphere
 	sphere.Center = Center;
 	sphere.W = 0.0f;
 
-	const TVertexAttributeArray<FVector>& VertexPositions = MeshDescription.VertexAttributes().GetAttributes<FVector>(MeshAttribute::Vertex::Position);
+	TVertexAttributesConstRef<FVector> VertexPositions = MeshDescription.VertexAttributes().GetAttributesRef<FVector>(MeshAttribute::Vertex::Position);
 
 	for (const FVertexID VertexID : MeshDescription.Vertices().GetElementIDs())
 	{
@@ -611,7 +611,7 @@ static void CalcBoundingSphyl(const FMeshDescription& MeshDescription, FSphere&
 	float r = Extents.GetMax();
 	float r2 = FMath::Square(r);
 	
-	const TVertexAttributeArray<FVector>& VertexPositions = MeshDescription.VertexAttributes().GetAttributes<FVector>(MeshAttribute::Vertex::Position);
+	TVertexAttributesConstRef<FVector> VertexPositions = MeshDescription.VertexAttributes().GetAttributesRef<FVector>(MeshAttribute::Vertex::Position);
 
 	// Now check each point lies within this the radius. If not - expand it a bit.
 	for (const FVertexID VertexID : MeshDescription.Vertices().GetElementIDs())
diff --git a/Engine/Source/Editor/UnrealEd/Private/StaticMeshEdit.cpp b/Engine/Source/Editor/UnrealEd/Private/StaticMeshEdit.cpp
--- a/Engine/Source/Editor/UnrealEd/Private/StaticMeshEdit.cpp
+++ b/Engine/Source/Editor/UnrealEd/Private/StaticMeshEdit.cpp
@@ -680,15 +680,15 @@ inline bool FVerticesEqual(FVector& V1,FVector& V2)
 
 void GetBrushMesh(ABrush* Brush, UModel* Model, FMeshDescription& MeshDescription, TArray<FStaticMaterial>& OutMaterials)
 {
-	TVertexAttributeArray<FVector>& VertexPositions = MeshDescription.VertexAttributes().GetAttributes<FVector>(MeshAttribute::Vertex::Position);
-	TVertexInstanceAttributeArray<FVector>& VertexInstanceNormals = MeshDescription.VertexInstanceAttributes().GetAttributes<FVector>(MeshAttribute::VertexInstance::Normal);
-	TVertexInstanceAttributeArray<FVector>& VertexInstanceTangents = MeshDescription.VertexInstanceAttributes().GetAttributes<FVector>(MeshAttribute::VertexInstance::Tangent);
-	TVertexInstanceAttributeArray<float>& VertexInstanceBinormalSigns = MeshDescription.VertexInstanceAttributes().GetAttributes<float>(MeshAttribute::VertexInstance::BinormalSign);
-	TVertexInstanceAttributeArray<FVector4>& VertexInstanceColors = MeshDescription.VertexInstanceAttributes().GetAttributes<FVector4>(MeshAttribute::VertexInstance::Color);
-	TVertexInstanceAttributeIndicesArray<FVector2D>& VertexInstanceUVs = MeshDescription.VertexInstanceAttributes().GetAttributesSet<FVector2D>(MeshAttribute::VertexInstance::TextureCoordinate);
-	TEdgeAttributeArray<bool>& EdgeHardnesses = MeshDescription.EdgeAttributes().GetAttributes<bool>(MeshAttribute::Edge::IsHard);
-	TEdgeAttributeArray<float>& EdgeCreaseSharpnesses = MeshDescription.EdgeAttributes().GetAttributes<float>(MeshAttribute::Edge::CreaseSharpness);
-	TPolygonGroupAttributeArray<FName>& PolygonGroupImportedMaterialSlotNames = MeshDescription.PolygonGroupAttributes().GetAttributes<FName>(MeshAttribute::PolygonGroup::ImportedMaterialSlotName);
+	TVertexAttributesRef<FVector> VertexPositions = MeshDescription.VertexAttributes().GetAttributesRef<FVector>(MeshAttribute::Vertex::Position);
+	TVertexInstanceAttributesRef<FVector> VertexInstanceNormals = MeshDescription.VertexInstanceAttributes().GetAttributesRef<FVector>(MeshAttribute::VertexInstance::Normal);
+	TVertexInstanceAttributesRef<FVector> VertexInstanceTangents = MeshDescription.VertexInstanceAttributes().GetAttributesRef<FVector>(MeshAttribute::VertexInstance::Tangent);
+	TVertexInstanceAttributesRef<float> VertexInstanceBinormalSigns = MeshDescription.VertexInstanceAttributes().GetAttributesRef<float>(MeshAttribute::VertexInstance::BinormalSign);
+	TVertexInstanceAttributesRef<FVector4> VertexInstanceColors = MeshDescription.VertexInstanceAttributes().GetAttributesRef<FVector4>(MeshAttribute::VertexInstance::Color);
+	TVertexInstanceAttributesRef<FVector2D> VertexInstanceUVs = MeshDescription.VertexInstanceAttributes().GetAttributesRef<FVector2D>(MeshAttribute::VertexInstance::TextureCoordinate);
+	TEdgeAttributesRef<bool> EdgeHardnesses = MeshDescription.EdgeAttributes().GetAttributesRef<bool>(MeshAttribute::Edge::IsHard);
+	TEdgeAttributesRef<float> EdgeCreaseSharpnesses = MeshDescription.EdgeAttributes().GetAttributesRef<float>(MeshAttribute::Edge::CreaseSharpness);
+	TPolygonGroupAttributesRef<FName> PolygonGroupImportedMaterialSlotNames = MeshDescription.PolygonGroupAttributes().GetAttributesRef<FName>(MeshAttribute::PolygonGroup::ImportedMaterialSlotName);
 	
 	//Make sure we have one UVChannel
 	VertexInstanceUVs.SetNumIndices(1);
@@ -766,9 +766,9 @@ void GetBrushMesh(ABrush* Brush, UModel* Model, FMeshDescription& MeshDescriptio
 					VertexPositions[VertexID[CornerIndex]] = Positions[CornerIndex];
 				}
 				VertexInstanceIDs[CornerIndex] = MeshDescription.CreateVertexInstance(VertexID[CornerIndex]);
-				VertexInstanceUVs.GetArrayForIndex(0)[VertexInstanceIDs[CornerIndex]] = FVector2D(
+				VertexInstanceUVs.Set(VertexInstanceIDs[CornerIndex], 0, FVector2D(
 					(Positions[CornerIndex] - TextureBase) | TextureX,
-					(Positions[CornerIndex] - TextureBase) | TextureY);
+					(Positions[CornerIndex] - TextureBase) | TextureY));
 			}
 
 			// Create a polygon with the 3 vertex instances
diff --git a/Engine/Source/Runtime/Core/Public/UObject/EditorObjectVersion.h b/Engine/Source/Runtime/Core/Public/UObject/EditorObjectVersion.h
--- a/Engine/Source/Runtime/Core/Public/UObject/EditorObjectVersion.h
+++ b/Engine/Source/Runtime/Core/Public/UObject/EditorObjectVersion.h
@@ -59,6 +59,10 @@ struct CORE_API FEditorObjectVersion
 		AddedMorphTargetSectionIndices,
 		// Serialize the instanced static mesh render data, to avoid building it at runtime
 		SerializeInstancedStaticMeshRenderData,
+		// Change to MeshDescription serialization (moved to release)
+		MeshDescriptionNewSerialization_MovedToRelease,
+		// New format for mesh description attributes
+		MeshDescriptionNewAttributeFormat,
 		// -----<new versions can be added above this line>-------------------------------------------------
 		VersionPlusOne,
 		LatestVersion = VersionPlusOne - 1
diff --git a/Engine/Source/Runtime/Engine/Private/StaticMesh.cpp b/Engine/Source/Runtime/Engine/Private/StaticMesh.cpp
--- a/Engine/Source/Runtime/Engine/Private/StaticMesh.cpp
+++ b/Engine/Source/Runtime/Engine/Private/StaticMesh.cpp
@@ -2614,6 +2614,27 @@ void UStaticMesh::RegisterMeshAttributes( FMeshDescription& MeshDescription )
 	MeshDescription.PolygonGroupAttributes().RegisterAttribute<FName>( MeshAttribute::PolygonGroup::ImportedMaterialSlotName ); //The unique key to match the mesh material slot
 	MeshDescription.PolygonGroupAttributes().RegisterAttribute<bool>( MeshAttribute::PolygonGroup::EnableCollision ); //Deprecated
 	MeshDescription.PolygonGroupAttributes().RegisterAttribute<bool>( MeshAttribute::PolygonGroup::CastShadow ); //Deprecated
+
+	{
+	TVertexAttributesView<FVector4> View = MeshDescription.VertexAttributes().GetAttributesView<FVector4>( MeshAttribute::Vertex::Position );
+	FVector4 Blah = View.Get(FVertexID(0));
+	(void)View.GetNumElements();
+	FVector Def = View.GetDefaultValue();
+	View.SetNumIndices(2);
+	View.Set(FVertexID(0), FVector4(1,1,1,1));
+	View.IsValid();
+	}
+
+	const FMeshDescription& MD2 = MeshDescription;
+	{
+	TVertexAttributesConstView<FVector4> View = MD2.VertexAttributes().GetAttributesView<FVector4>( MeshAttribute::Vertex::Position );
+	FVector4 Blah = View.Get(FVertexID(0));
+	(void)View.GetNumElements();
+	FVector Def = View.GetDefaultValue();
+//	View.SetNumIndices(2);
+//	View.Set(FVertexID(0), FVector4(1,1,1,1));
+	View.IsValid();
+	}
 }
 
 
@@ -2956,7 +2977,7 @@ void UStaticMesh::FixupMaterialSlotName()
 // differences, etc.) replace the version GUID below with a new one.
 // In case of merge conflicts with DDC versions, you MUST generate a new GUID
 // and set this new GUID as the version.                                       
-#define MESHDATAKEY_STATICMESH_DERIVEDDATA_VER TEXT("A3E9E442F5784050BCAF878E4E80EE44")
+#define MESHDATAKEY_STATICMESH_DERIVEDDATA_VER TEXT("B6081AC5FA9949A7B87DC14E39BC25E6")
 
 static const FString& GetMeshDataKeyStaticMeshDerivedDataVersion()
 {
diff --git a/Engine/Source/Runtime/MeshDescription/Private/MeshDescription.cpp b/Engine/Source/Runtime/MeshDescription/Private/MeshDescription.cpp
--- a/Engine/Source/Runtime/MeshDescription/Private/MeshDescription.cpp
+++ b/Engine/Source/Runtime/MeshDescription/Private/MeshDescription.cpp
@@ -321,7 +321,7 @@ FPlane FMeshDescription::ComputePolygonPlane( const FPolygonID PolygonID ) const
 	GetPolygonPerimeterVertices( PolygonID, /* Out */ PerimeterVertexIDs );
 
 	// @todo Maybe this shouldn't be in FMeshDescription but in a utility class, as it references a specific attribute name
-	const TVertexAttributeArray<FVector>& VertexPositions = VertexAttributes().GetAttributes<FVector>( MeshAttribute::Vertex::Position );
+	TVertexAttributesConstRef<FVector> VertexPositions = VertexAttributes().GetAttributesRef<FVector>( MeshAttribute::Vertex::Position );
 
 	// Use 'Newell's Method' to compute a robust 'best fit' plane from the vertices of this polygon
 	for ( int32 VertexNumberI = PerimeterVertexIDs.Num() - 1, VertexNumberJ = 0; VertexNumberJ < PerimeterVertexIDs.Num(); VertexNumberI = VertexNumberJ, VertexNumberJ++ )
@@ -411,7 +411,7 @@ void FMeshDescription::ComputePolygonTriangulation(const FPolygonID PolygonID, T
 	static TArray<FVector> VertexPositions;
 
 	{
-		const TVertexAttributeArray<FVector>& MeshVertexPositions = VertexAttributes().GetAttributes<FVector>(MeshAttribute::Vertex::Position);
+		const TVertexAttributesRef<FVector> MeshVertexPositions = VertexAttributes().GetAttributesRef<FVector>(MeshAttribute::Vertex::Position);
 		PrevVertexNumbers.SetNumUninitialized(PolygonVertexCount, false);
 		NextVertexNumbers.SetNumUninitialized(PolygonVertexCount, false);
 		VertexPositions.SetNumUninitialized(PolygonVertexCount, false);
@@ -532,12 +532,12 @@ void FMeshDescription::TriangulateMesh()
 
 bool FMeshDescription::ComputePolygonTangentsAndNormals(const FPolygonID PolygonID
 	, float ComparisonThreshold
-	, const TVertexAttributeArray<FVector>& VertexPositions
-	, const TVertexInstanceAttributeArray<FVector2D>& VertexUVs
-	, TPolygonAttributeArray<FVector>& PolygonNormals
-	, TPolygonAttributeArray<FVector>& PolygonTangents
-	, TPolygonAttributeArray<FVector>& PolygonBinormals
-	, TPolygonAttributeArray<FVector>& PolygonCenters)
+	, const TVertexAttributesRef<FVector> VertexPositions
+	, const TVertexInstanceAttributesRef<FVector2D> VertexUVs
+	, TPolygonAttributesRef<FVector> PolygonNormals
+	, TPolygonAttributesRef<FVector> PolygonTangents
+	, TPolygonAttributesRef<FVector> PolygonBinormals
+	, TPolygonAttributesRef<FVector> PolygonCenters)
 {
 	bool bValidNTBs = true;
 	// Calculate the center of this polygon
@@ -597,12 +597,12 @@ bool FMeshDescription::ComputePolygonTangentsAndNormals(const FPolygonID Polygon
 
 void FMeshDescription::ComputePolygonTangentsAndNormals(const TArray<FPolygonID>& PolygonIDs, float ComparisonThreshold)
 {
-	const TVertexAttributeArray<FVector>& VertexPositions = VertexAttributes().GetAttributes<FVector>(MeshAttribute::Vertex::Position);
-	const TVertexInstanceAttributeArray<FVector2D>& VertexUVs = VertexInstanceAttributes().GetAttributes<FVector2D>(MeshAttribute::VertexInstance::TextureCoordinate, 0);
-	TPolygonAttributeArray<FVector>& PolygonNormals = PolygonAttributes().GetAttributes<FVector>(MeshAttribute::Polygon::Normal);
-	TPolygonAttributeArray<FVector>& PolygonTangents = PolygonAttributes().GetAttributes<FVector>(MeshAttribute::Polygon::Tangent);
-	TPolygonAttributeArray<FVector>& PolygonBinormals = PolygonAttributes().GetAttributes<FVector>(MeshAttribute::Polygon::Binormal);
-	TPolygonAttributeArray<FVector>& PolygonCenters = PolygonAttributes().GetAttributes<FVector>(MeshAttribute::Polygon::Center);
+	const TVertexAttributesRef<FVector> VertexPositions = VertexAttributes().GetAttributesRef<FVector>(MeshAttribute::Vertex::Position);
+	const TVertexInstanceAttributesRef<FVector2D> VertexUVs = VertexInstanceAttributes().GetAttributesRef<FVector2D>(MeshAttribute::VertexInstance::TextureCoordinate);
+	TPolygonAttributesRef<FVector> PolygonNormals = PolygonAttributes().GetAttributesRef<FVector>(MeshAttribute::Polygon::Normal);
+	TPolygonAttributesRef<FVector> PolygonTangents = PolygonAttributes().GetAttributesRef<FVector>(MeshAttribute::Polygon::Tangent);
+	TPolygonAttributesRef<FVector> PolygonBinormals = PolygonAttributes().GetAttributesRef<FVector>(MeshAttribute::Polygon::Binormal);
+	TPolygonAttributesRef<FVector> PolygonCenters = PolygonAttributes().GetAttributesRef<FVector>(MeshAttribute::Polygon::Center);
 
 	TArray<FPolygonID> DegeneratePolygonIDs;
 	for (const FPolygonID PolygonID : PolygonIDs)
@@ -620,23 +620,23 @@ void FMeshDescription::ComputePolygonTangentsAndNormals(const TArray<FPolygonID>
 		TArray<FVertexInstanceID> OrphanedVertexInstances;
 		TArray<FPolygonGroupID> OrphanedPolygonGroups;
 		TArray<FVertexID> OrphanedVertices;
-		for (FPolygonID& PolygonID : DegeneratePolygonIDs)
+		for (FPolygonID PolygonID : DegeneratePolygonIDs)
 		{
 			DeletePolygon(PolygonID, &OrphanedEdges, &OrphanedVertexInstances, &OrphanedPolygonGroups);
 		}
-		for (FPolygonGroupID& PolygonGroupID : OrphanedPolygonGroups)
+		for (FPolygonGroupID PolygonGroupID : OrphanedPolygonGroups)
 		{
 			DeletePolygonGroup(PolygonGroupID);
 		}
-		for (FVertexInstanceID& VertexInstanceID : OrphanedVertexInstances)
+		for (FVertexInstanceID VertexInstanceID : OrphanedVertexInstances)
 		{
 			DeleteVertexInstance(VertexInstanceID, &OrphanedVertices);
 		}
-		for (FEdgeID& EdgeID : OrphanedEdges)
+		for (FEdgeID EdgeID : OrphanedEdges)
 		{
 			DeleteEdge(EdgeID, &OrphanedVertices);
 		}
-		for (FVertexID& VertexID : OrphanedVertices)
+		for (FVertexID VertexID : OrphanedVertices)
 		{
 			DeleteVertex(VertexID);
 		}
@@ -664,7 +664,7 @@ void FMeshDescription::GetConnectedSoftEdges(const FVertexID VertexID, TArray<FE
 {
 	OutConnectedSoftEdges.Reset();
 
-	const TEdgeAttributeArray<bool>& EdgeHardnesses = EdgeAttributes().GetAttributes<bool>(MeshAttribute::Edge::IsHard);
+	TEdgeAttributesConstRef<bool> EdgeHardnesses = EdgeAttributes().GetAttributesRef<bool>(MeshAttribute::Edge::IsHard);
 	for (const FEdgeID ConnectedEdgeID : GetVertex(VertexID).ConnectedEdgeIDs)
 	{
 		if (!EdgeHardnesses[ConnectedEdgeID])
@@ -755,7 +755,7 @@ float FMeshDescription::GetPolygonCornerAngleForVertex(const FPolygonID PolygonI
 		const FVertexID ThisVertexID = GetVertexInstanceVertex(Contour.VertexInstanceIDs[ContourIndex]);
 		const FVertexID NextVertexID = GetVertexInstanceVertex(Contour.VertexInstanceIDs[NextIndex]);
 
-		const TVertexAttributeArray<FVector>& VertexPositions = VertexAttributes().GetAttributes<FVector>(MeshAttribute::Vertex::Position);
+		TVertexAttributesConstRef<FVector> VertexPositions = VertexAttributes().GetAttributesRef<FVector>(MeshAttribute::Vertex::Position);
 
 		const FVector PrevVertexPosition = VertexPositions[PrevVertexID];
 		const FVector ThisVertexPosition = VertexPositions[ThisVertexID];
@@ -799,14 +799,14 @@ float FMeshDescription::GetPolygonCornerAngleForVertex(const FPolygonID PolygonI
 }
 
 
-void FMeshDescription::ComputeTangentsAndNormals(const FVertexInstanceID& VertexInstanceID
+void FMeshDescription::ComputeTangentsAndNormals(const FVertexInstanceID VertexInstanceID
 	, EComputeNTBsOptions ComputeNTBsOptions
-	, const TPolygonAttributeArray<FVector>& PolygonNormals
-	, const TPolygonAttributeArray<FVector>& PolygonTangents
-	, const TPolygonAttributeArray<FVector>& PolygonBinormals
-	, TVertexInstanceAttributeArray<FVector>& VertexNormals
-	, TVertexInstanceAttributeArray<FVector>& VertexTangents
-	, TVertexInstanceAttributeArray<float>& VertexBinormalSigns)
+	, const TPolygonAttributesRef<FVector> PolygonNormals
+	, const TPolygonAttributesRef<FVector> PolygonTangents
+	, const TPolygonAttributesRef<FVector> PolygonBinormals
+	, TVertexInstanceAttributesRef<FVector> VertexNormals
+	, TVertexInstanceAttributesRef<FVector> VertexTangents
+	, TVertexInstanceAttributesRef<float> VertexBinormalSigns)
 {
 	bool bComputeNormals = !!(ComputeNTBsOptions & EComputeNTBsOptions::Normals);
 	bool bComputeTangents = !!(ComputeNTBsOptions & EComputeNTBsOptions::Tangents);
@@ -911,13 +911,13 @@ void FMeshDescription::ComputeTangentsAndNormals(const FVertexInstanceID& Vertex
 
 void FMeshDescription::ComputeTangentsAndNormals(const TArray<FVertexInstanceID>& VertexInstanceIDs, EComputeNTBsOptions ComputeNTBsOptions)
 {
-	const TPolygonAttributeArray<FVector>& PolygonNormals = PolygonAttributes().GetAttributes<FVector>(MeshAttribute::Polygon::Normal);
-	const TPolygonAttributeArray<FVector>& PolygonTangents = PolygonAttributes().GetAttributes<FVector>(MeshAttribute::Polygon::Tangent);
-	const TPolygonAttributeArray<FVector>& PolygonBinormals = PolygonAttributes().GetAttributes<FVector>(MeshAttribute::Polygon::Binormal);
+	const TPolygonAttributesRef<FVector> PolygonNormals = PolygonAttributes().GetAttributesRef<FVector>(MeshAttribute::Polygon::Normal);
+	const TPolygonAttributesRef<FVector> PolygonTangents = PolygonAttributes().GetAttributesRef<FVector>(MeshAttribute::Polygon::Tangent);
+	const TPolygonAttributesRef<FVector> PolygonBinormals = PolygonAttributes().GetAttributesRef<FVector>(MeshAttribute::Polygon::Binormal);
 
-	TVertexInstanceAttributeArray<FVector>& VertexNormals = VertexInstanceAttributes().GetAttributes<FVector>(MeshAttribute::VertexInstance::Normal, 0);
-	TVertexInstanceAttributeArray<FVector>& VertexTangents = VertexInstanceAttributes().GetAttributes<FVector>(MeshAttribute::VertexInstance::Tangent, 0);
-	TVertexInstanceAttributeArray<float>& VertexBinormalSigns = VertexInstanceAttributes().GetAttributes<float>(MeshAttribute::VertexInstance::BinormalSign, 0);
+	TVertexInstanceAttributesRef<FVector> VertexNormals = VertexInstanceAttributes().GetAttributesRef<FVector>(MeshAttribute::VertexInstance::Normal);
+	TVertexInstanceAttributesRef<FVector> VertexTangents = VertexInstanceAttributes().GetAttributesRef<FVector>(MeshAttribute::VertexInstance::Tangent);
+	TVertexInstanceAttributesRef<float> VertexBinormalSigns = VertexInstanceAttributes().GetAttributesRef<float>(MeshAttribute::VertexInstance::BinormalSign);
 
 	for (const FVertexInstanceID VertexInstanceID : VertexInstanceIDs)
 	{
@@ -928,13 +928,13 @@ void FMeshDescription::ComputeTangentsAndNormals(const TArray<FVertexInstanceID>
 
 void FMeshDescription::ComputeTangentsAndNormals(EComputeNTBsOptions ComputeNTBsOptions)
 {
-	const TPolygonAttributeArray<FVector>& PolygonNormals = PolygonAttributes().GetAttributes<FVector>(MeshAttribute::Polygon::Normal);
-	const TPolygonAttributeArray<FVector>& PolygonTangents = PolygonAttributes().GetAttributes<FVector>(MeshAttribute::Polygon::Tangent);
-	const TPolygonAttributeArray<FVector>& PolygonBinormals = PolygonAttributes().GetAttributes<FVector>(MeshAttribute::Polygon::Binormal);
+	const TPolygonAttributesRef<FVector> PolygonNormals = PolygonAttributes().GetAttributesRef<FVector>(MeshAttribute::Polygon::Normal);
+	const TPolygonAttributesRef<FVector> PolygonTangents = PolygonAttributes().GetAttributesRef<FVector>(MeshAttribute::Polygon::Tangent);
+	const TPolygonAttributesRef<FVector> PolygonBinormals = PolygonAttributes().GetAttributesRef<FVector>(MeshAttribute::Polygon::Binormal);
 
-	TVertexInstanceAttributeArray<FVector>& VertexNormals = VertexInstanceAttributes().GetAttributes<FVector>(MeshAttribute::VertexInstance::Normal, 0);
-	TVertexInstanceAttributeArray<FVector>& VertexTangents = VertexInstanceAttributes().GetAttributes<FVector>(MeshAttribute::VertexInstance::Tangent, 0);
-	TVertexInstanceAttributeArray<float>& VertexBinormalSigns = VertexInstanceAttributes().GetAttributes<float>(MeshAttribute::VertexInstance::BinormalSign, 0);
+	TVertexInstanceAttributesRef<FVector> VertexNormals = VertexInstanceAttributes().GetAttributesRef<FVector>(MeshAttribute::VertexInstance::Normal);
+	TVertexInstanceAttributesRef<FVector> VertexTangents = VertexInstanceAttributes().GetAttributesRef<FVector>(MeshAttribute::VertexInstance::Tangent);
+	TVertexInstanceAttributesRef<float> VertexBinormalSigns = VertexInstanceAttributes().GetAttributesRef<float>(MeshAttribute::VertexInstance::BinormalSign);
 
 	for (const FVertexInstanceID VertexInstanceID : VertexInstances().GetElementIDs())
 	{
@@ -945,8 +945,8 @@ void FMeshDescription::ComputeTangentsAndNormals(EComputeNTBsOptions ComputeNTBs
 
 void FMeshDescription::DetermineEdgeHardnessesFromVertexInstanceNormals( const float Tolerance )
 {
-	const TVertexInstanceAttributeArray<FVector>& VertexNormals = VertexInstanceAttributes().GetAttributes<FVector>( MeshAttribute::VertexInstance::Normal );
-	TEdgeAttributeArray<bool>& EdgeHardnesses = EdgeAttributes().GetAttributes<bool>( MeshAttribute::Edge::IsHard );
+	const TVertexInstanceAttributesRef<FVector> VertexNormals = VertexInstanceAttributes().GetAttributesRef<FVector>( MeshAttribute::VertexInstance::Normal );
+	TEdgeAttributesRef<bool> EdgeHardnesses = EdgeAttributes().GetAttributesRef<bool>( MeshAttribute::Edge::IsHard );
 
 	// Holds unique vertex instance IDs for a given edge vertex
 	// @todo: use TMemStackAllocator or similar to avoid expensive allocations
@@ -1009,8 +1009,8 @@ void FMeshDescription::DetermineEdgeHardnessesFromVertexInstanceNormals( const f
 
 void FMeshDescription::DetermineUVSeamsFromUVs( const int32 UVIndex, const float Tolerance )
 {
-	const TVertexInstanceAttributeArray<FVector2D>& VertexUVs = VertexInstanceAttributes().GetAttributes<FVector2D>( MeshAttribute::VertexInstance::TextureCoordinate, UVIndex );
-	TEdgeAttributeArray<bool>& EdgeUVSeams = EdgeAttributes().GetAttributes<bool>( MeshAttribute::Edge::IsUVSeam );
+	const TVertexInstanceAttributesRef<FVector2D> VertexUVs = VertexInstanceAttributes().GetAttributesRef<FVector2D>( MeshAttribute::VertexInstance::TextureCoordinate );
+	TEdgeAttributesRef<bool> EdgeUVSeams = EdgeAttributes().GetAttributesRef<bool>( MeshAttribute::Edge::IsUVSeam );
 
 	// Holds unique vertex instance IDs for a given edge vertex
 	// @todo: use TMemStackAllocator or similar to avoid expensive allocations
@@ -1054,10 +1054,10 @@ void FMeshDescription::DetermineUVSeamsFromUVs( const int32 UVIndex, const float
 
 			// First unique vertex instance is used as a reference against which the others are compared.
 			// (not a perfect approach: really the 'median' should be used as a reference)
-			const FVector2D ReferenceUV = VertexUVs[ UniqueVertexInstanceIDs[ 0 ] ];
+			const FVector2D ReferenceUV = VertexUVs.Get( UniqueVertexInstanceIDs[ 0 ], UVIndex );
 			for( int32 Index = 1; Index < UniqueVertexInstanceIDs.Num(); ++Index )
 			{
-				if( !VertexUVs[ UniqueVertexInstanceIDs[ Index ] ].Equals( ReferenceUV, Tolerance ) )
+				if( !VertexUVs.Get( UniqueVertexInstanceIDs[ Index ], UVIndex ).Equals( ReferenceUV, Tolerance ) )
 				{
 					bEdgeIsUVSeam = true;
 					break;
@@ -1072,7 +1072,7 @@ void FMeshDescription::DetermineUVSeamsFromUVs( const int32 UVIndex, const float
 
 void FMeshDescription::GetPolygonsInSameChartAsPolygon( const FPolygonID PolygonID, TArray<FPolygonID>& OutPolygonIDs )
 {
-	const TEdgeAttributeArray<bool>& EdgeUVSeams = EdgeAttributes().GetAttributes<bool>( MeshAttribute::Edge::IsUVSeam );
+	const TEdgeAttributesRef<bool> EdgeUVSeams = EdgeAttributes().GetAttributesRef<bool>( MeshAttribute::Edge::IsUVSeam );
 	const int32 NumPolygons = Polygons().Num();
 
 	// This holds the results - all polygon IDs which are in the same UV chart
diff --git a/Engine/Source/Runtime/MeshDescription/Public/MeshAttributeArray.h b/Engine/Source/Runtime/MeshDescription/Public/MeshAttributeArray.h
--- a/Engine/Source/Runtime/MeshDescription/Public/MeshAttributeArray.h
+++ b/Engine/Source/Runtime/MeshDescription/Public/MeshAttributeArray.h
@@ -10,12 +10,9 @@
 
 /**
  * List of attribute types which are supported.
- * We do this so we can automatically generate the attribute containers and their associated accessors with
- * some template magic.
  *
- * IMPORTANT NOTE: Do not remove any type from this tuple, or serialization will fail.
- * Types may be added at the end of this list if necessary, although please do so sparingly as each extra type will
- * impact on performance and object size.
+ * IMPORTANT NOTE: Do not reorder or remove any type from this tuple, or serialization will fail.
+ * Types may be added at the end of this list as required.
  */
 using AttributeTypes = TTuple
 <
@@ -29,6 +26,50 @@ using AttributeTypes = TTuple
 >;
 
 
+/**
+ * Helper template which gets the tuple index of a given type from a given TTuple.
+ * If the type occurs more than once, the first index is returned.
+ * If the type doesn't appear, a compile error is generated.
+ *
+ * Given Type = char, and Tuple = TTuple<int, float, char>,
+ * TTupleIndex<Type, Tuple>::Value will be 2.
+ *
+ * @todo: Move to Tuple.h
+ */
+template <typename Type, typename Tuple> struct TTupleIndex;
+template <typename Type, typename... Ts> struct TTupleIndex<Type, TTuple<Type, Ts...>> : TIntegralConstant<uint32, 0> {};
+template <typename Type, typename... Ts> struct TTupleIndex<Type, TTuple<Ts...>> : TIntegralConstant<uint32, 0> { static_assert( sizeof...( Ts ) > 0, "TTuple type not found" ); };
+template <typename Type, typename T, typename... Ts> struct TTupleIndex<Type, TTuple<T, Ts...>> : TIntegralConstant<uint32, 1 + TTupleIndex<Type, TTuple<Ts...>>::Value> {};
+
+
+/**
+ * Helper template which gets the element type of a TTuple with the given index.
+ *
+ * Given Index = 1, and Tuple = TTuple<int, float, char>,
+ * TTupleElement<Index, Tuple>::Type will be float.
+ *
+ * @todo: Move to Tuple.h
+ */
+template <uint32 Index, typename Tuple> struct TTupleElement;
+template <uint32 Index, typename T, typename... Ts> struct TTupleElement<Index, TTuple<T, Ts...>> : TTupleElement<Index - 1, TTuple<Ts...>> {};
+template <typename T, typename... Ts> struct TTupleElement<0, TTuple<T, Ts...>> { using Type = T; };
+template <typename... Ts> struct TTupleElement<0, TTuple<Ts...>> { static_assert( sizeof...( Ts ) > 0, "TTuple element index out of range" ); };
+
+
+/**
+ * Class which implements a function jump table to be automatically generated at compile time.
+ * This is used by TAttributesSet to provide O(1) dispatch by attribute type at runtime.
+ */
+template <typename FnType, uint32 Size>
+struct TJumpTable
+{
+	template <typename... T>
+	explicit constexpr TJumpTable( T... Ts ) : Fns{ Ts... } {}
+
+	FnType* Fns[Size];
+};
+
+
 /**
  * Traits class to specify which attribute types can be bulk serialized.
  */
@@ -43,63 +84,80 @@ template <> struct TIsBulkSerializable<FName> { static const bool Value = false;
  * Note that the container may grow arbitrarily as new elements are inserted, but it will never be
  * shrunk as elements are removed. The only operations that will shrink the container are Initialize() and Remap().
  */
-template <typename ElementType>
+template <typename AttributeType>
 class TMeshAttributeArrayBase
 {
 public:
 
-	/**
-	 * Custom serialization for TMeshAttributeArrayBase.
-	 */
-	template <typename T>
-	friend typename TEnableIf<!TIsBulkSerializable<T>::Value, FArchive>::Type& operator<<( FArchive& Ar, TMeshAttributeArrayBase<T>& Array );
+	/** Custom serialization for TMeshAttributeArrayBase. */
+	template <typename T> friend typename TEnableIf<!TIsBulkSerializable<T>::Value, FArchive>::Type& operator<<( FArchive& Ar, TMeshAttributeArrayBase<T>& Array );
+	template <typename T> friend typename TEnableIf<TIsBulkSerializable<T>::Value, FArchive>::Type& operator<<( FArchive& Ar, TMeshAttributeArrayBase<T>& Array );
 
-	template <typename T>
-	friend typename TEnableIf<TIsBulkSerializable<T>::Value, FArchive>::Type& operator<<( FArchive& Ar, TMeshAttributeArrayBase<T>& Array );
+	/** Return size of container */
+	FORCEINLINE int32 Num() const { return Container.Num(); }
 
-protected:
+	/** Return base of data */
+	FORCEINLINE const AttributeType* GetData() const { return Container.GetData(); }
+
+	/** Initializes the array to the given size with the default value */
+	FORCEINLINE void Initialize( const int32 ElementCount, const AttributeType& Default )
+	{
+		Container.Reset( ElementCount );
+		Insert( ElementCount - 1, Default );
+	}
 
-	friend class FMeshDescription;
+	/** Expands the array if necessary so that the passed element index is valid. Newly created elements will be assigned the default value. */
+	void Insert( const int32 Index, const AttributeType& Default );
 
-	/** Should not instance this base class directly */
-	TMeshAttributeArrayBase() = default;
+	/** Remaps elements according to the passed remapping table */
+	void Remap( const TSparseArray<int32>& IndexRemap, const AttributeType& Default );
 
-	/** Disallow use of the copy constructor outside of FMeshDescription, to prevent arrays being mistakenly accessed in the UMeshDescription by value */
-	TMeshAttributeArrayBase( const TMeshAttributeArrayBase& ) = default;
-	TMeshAttributeArrayBase( TMeshAttributeArrayBase&& ) = default;
-	TMeshAttributeArrayBase& operator=( const TMeshAttributeArrayBase& ) = default;
-	TMeshAttributeArrayBase& operator=( TMeshAttributeArrayBase&& ) = default;
+	/** Element accessors */
+	FORCEINLINE const AttributeType& operator[]( const int32 Index ) const { return Container[ Index ]; }
+	FORCEINLINE AttributeType& operator[]( const int32 Index ) { return Container[ Index ]; }
 
-	/** Expands the array if necessary so that the passed element index is valid. Newly created elements will be assigned the default value. */
-	void Insert( const int32 Index, const ElementType& Default )
+protected:
+
+	/** The actual container, represented by a regular array */
+	TArray<AttributeType> Container;
+};
+
+
+template <typename AttributeType>
+void TMeshAttributeArrayBase<AttributeType>::Insert( const int32 Index, const AttributeType& Default )
+{
+	if( Index >= Container.Num() )
 	{
-		if( Index >= Container.Num() )
-		{
-			// If the index is off the end of the container, add as many elements as required to make it the last valid index.
-			int32 StartIndex = Container.AddUninitialized( Index + 1 - Container.Num() );
-			ElementType* Data = Container.GetData() + StartIndex;
+		// If the index is off the end of the container, add as many elements as required to make it the last valid index.
+		int32 StartIndex = Container.AddUninitialized( Index + 1 - Container.Num() );
+		AttributeType* Data = Container.GetData() + StartIndex;
 
-			// Construct added elements with the default value passed in
-			while( StartIndex <= Index )
-			{
-				new( Data ) ElementType( Default );
-				StartIndex++;
-				Data++;
-			}
+		// Construct added elements with the default value passed in
+		while( StartIndex <= Index )
+		{
+			new( Data ) AttributeType( Default );
+			StartIndex++;
+			Data++;
 		}
 	}
+}
 
-	/** Initializes the array to the given size with the default value */
-	void Initialize( const int32 ElementCount, const ElementType& Default )
+template <typename AttributeType>
+void TMeshAttributeArrayBase<AttributeType>::Remap( const TSparseArray<int32>& IndexRemap, const AttributeType& Default )
+{
+	TMeshAttributeArrayBase NewAttributeArray;
+
+	for( typename TSparseArray<int32>::TConstIterator It( IndexRemap ); It; ++It )
 	{
-		Container.Reset( ElementCount );
-		Insert( ElementCount - 1, Default );
-	}
+		const int32 OldElementIndex = It.GetIndex();
+		const int32 NewElementIndex = IndexRemap[ OldElementIndex ];
 
-	/** The actual container, represented by a regular array */
-	TArray<ElementType> Container;
-};
+		NewAttributeArray.Insert( NewElementIndex, Default );
+		NewAttributeArray[ NewElementIndex ] = MoveTemp( Container[ OldElementIndex ] );
+	}
 
+	Container = MoveTemp( NewAttributeArray.Container );
+}
 
 template <typename T>
 inline typename TEnableIf<!TIsBulkSerializable<T>::Value, FArchive>::Type& operator<<( FArchive& Ar, TMeshAttributeArrayBase<T>& Array )
@@ -127,161 +185,178 @@ inline typename TEnableIf<TIsBulkSerializable<T>::Value, FArchive>::Type& operat
 }
 
 
-
+// This is a deprecated class which will be removed in 4.21
 template <typename AttributeType, typename ElementIDType>
-class TAttributeIndicesArray;
-
-/**
- * We prefer to access elements of the container via strongly-typed IDs.
- * This derived class imposes this type safety.
- */
-template <typename ElementType, typename ElementIDType>
-class TMeshAttributeArray : private TMeshAttributeArrayBase<ElementType>
+class TMeshAttributeArray final : public TMeshAttributeArrayBase<AttributeType>
 {
 	static_assert( TIsDerivedFrom<ElementIDType, FElementID>::IsDerived, "ElementIDType must be derived from FElementID" );
 
-	using TMeshAttributeArrayBase<ElementType>::Container;
+	using Super = TMeshAttributeArrayBase<AttributeType>;
+	using Super::operator[];
 
 public:
 
 	/** Element accessors */
-	FORCEINLINE const ElementType& operator[]( const ElementIDType ElementID ) const { return Container[ ElementID.GetValue() ]; }
-	FORCEINLINE ElementType& operator[]( const ElementIDType ElementID ) { return Container[ ElementID.GetValue() ]; }
-
-	/** Return size of container */
-	FORCEINLINE int32 Num() const { return Container.Num(); }
-
-	/** Return base of data */
-	FORCEINLINE const ElementType* GetData() const { return Container.GetData(); }
-
-protected:
-
-	friend class TAttributeIndicesArray<ElementType, ElementIDType>;
-
-	/** Expands the array if necessary so that the passed element index is valid. Newly created elements will be assigned the default value. */
-	FORCEINLINE void Insert( const ElementIDType Index, const ElementType& Default )
-	{
-		TMeshAttributeArrayBase<ElementType>::Insert( Index.GetValue(), Default );
-	}
-
-	/** Serializer */
-	friend FORCEINLINE FArchive& operator<<( FArchive& Ar, TMeshAttributeArray& Array )
-	{
-		Ar << static_cast<TMeshAttributeArrayBase<ElementType>&>( Array );
-		return Ar;
-	}
+	DEPRECATED( 4.20, "Please access attributes via GetAttributesRef() or GetAttributesView() instead." )
+	FORCEINLINE const AttributeType& operator[]( const ElementIDType ElementID ) const { return this->Container[ ElementID.GetValue() ]; }
 
-	/** Remaps elements according to the passed remapping table */
-	void Remap( const TSparseArray<ElementIDType>& IndexRemap, const ElementType& Default );
+	DEPRECATED( 4.20, "Please access attributes via GetAttributesRef() or GetAttributesView() instead." )
+	FORCEINLINE AttributeType& operator[]( const ElementIDType ElementID ) { return this->Container[ ElementID.GetValue() ]; }
 };
 
 
-template <typename ElementType, typename ElementIDType>
-void TMeshAttributeArray<ElementType, ElementIDType>::Remap( const TSparseArray<ElementIDType>& IndexRemap, const ElementType& Default )
-{
-	TMeshAttributeArray NewAttributeArray;
-
-	for( typename TSparseArray<ElementIDType>::TConstIterator It( IndexRemap ); It; ++It )
-	{
-		const int32 OldElementIndex = It.GetIndex();
-		const ElementIDType NewElementIndex = IndexRemap[ OldElementIndex ];
-
-		NewAttributeArray.Insert( NewElementIndex, Default );
-		NewAttributeArray[ NewElementIndex ] = MoveTemp( Container[ OldElementIndex ] );
-	}
-
-	Container = MoveTemp( NewAttributeArray.Container );
-}
-
-
-/** Define aliases for element attributes */
-template <typename AttributeType> using TVertexAttributeArray = TMeshAttributeArray<AttributeType, FVertexID>;
-template <typename AttributeType> using TVertexInstanceAttributeArray = TMeshAttributeArray<AttributeType, FVertexInstanceID>;
-template <typename AttributeType> using TEdgeAttributeArray = TMeshAttributeArray<AttributeType, FEdgeID>;
-template <typename AttributeType> using TPolygonAttributeArray = TMeshAttributeArray<AttributeType, FPolygonID>;
-template <typename AttributeType> using TPolygonGroupAttributeArray = TMeshAttributeArray<AttributeType, FPolygonGroupID>;
-
 
 /**
  * Flags specifying properties of an attribute
- * @todo mesh description: this needs to be moved to a application-specific place;
- * this code is too low-level to assume particular meanings for the flags.
  */
 enum class EMeshAttributeFlags : uint32
 {
 	None				= 0,
-	Lerpable			= ( 1 << 0 ),
-	AutoGenerated		= ( 1 << 1 ),
-	Mergeable			= ( 1 << 2 )
+	Lerpable			= ( 1 << 0 ),	/** Attribute can be automatically lerped according to the value of 2 or 3 other attributes */
+	AutoGenerated		= ( 1 << 1 ),	/** Attribute is auto-generated by importer or editable mesh, rather than representing an imported property */	
+	Mergeable			= ( 1 << 2 ),	/** If all vertices' attributes are mergeable, and of near-equal value, they can be welded */
+	Transient			= ( 1 << 3 )	/** Attribute is not serialized */
 };
 
 ENUM_CLASS_FLAGS( EMeshAttributeFlags );
 
 
 /**
- * This class represents a container for a named attribute on a mesh element.
- * It contains an array of TMeshAttributeArrays, one per attribute index. 
+ * This is the base class for an attribute array set.
+ * An attribute array set is a container which holds attribute arrays, one per attribute index.
+ * Many attributes have only one index, while others (such as texture coordinates) may want to define many.
+ *
+ * All attribute array set instances will be of derived types; this type exists for polymorphism purposes,
+ * so that they can be managed by a generic TUniquePtr<FMeshAttributeArraySetBase>.
+ *
+ * In general, we avoid accessing them via virtual dispatch by insisting that their type be passed as
+ * a template parameter in the accessor. This can be checked against the Type field to ensure that we are
+ * accessing an instance by its correct type.
  */
-template <typename T, typename U>
-class TAttributeIndicesArray
+class FMeshAttributeArraySetBase
 {
 public:
+	/** Constructor */
+	FORCEINLINE FMeshAttributeArraySetBase( const uint32 InType, const EMeshAttributeFlags InFlags, const int32 InNumberOfElements )
+		: Type( InType ),
+		  NumElements( InNumberOfElements ),
+		  Flags( InFlags )
+	{}
 
-	using AttributeType = T;
-	using ElementIDType = U;
+	/** Virtual interface */
+	virtual ~FMeshAttributeArraySetBase() = default;
+	virtual TUniquePtr<FMeshAttributeArraySetBase> Clone() const = 0;
+	virtual void Insert( const int32 Index ) = 0;
+	virtual void Remove( const int32 Index ) = 0;
+	virtual void Initialize( const int32 Count ) = 0;
+	virtual void Serialize( FArchive& Ar ) = 0;
+	virtual void Remap( const TSparseArray<int32>& IndexRemap ) = 0;
+	virtual int32 GetNumIndices() const = 0;
+	virtual void SetNumIndices( const int32 NumIndices ) = 0;
+
+	/** Determine whether this attribute array set is of the given type */
+	template <typename T>
+	FORCEINLINE bool HasType() const
+	{
+		return TTupleIndex<T, AttributeTypes>::Value == Type;
+	}
+
+	/** Get the type index of this attribute array set */
+	FORCEINLINE uint32 GetType() const { return Type; }
+
+	/** Get the flags for this attribute array set */
+	FORCEINLINE EMeshAttributeFlags GetFlags() const { return Flags; }
+
+	/** Return number of elements each attribute index has */
+	FORCEINLINE int32 GetNumElements() const { return NumElements; }
+
+protected:
+	/** Type of the attribute array (based on the tuple element index from AttributeTypes) */
+	uint32 Type;
+
+	/** Number of elements in each index */
+	int32 NumElements;
+
+	/** Implementation-defined attribute name flags */
+	EMeshAttributeFlags Flags;
+};
 
-	/** Default constructor - required so that it builds correctly */
-	TAttributeIndicesArray() = default;
 
+/**
+ * This is a type-specific attribute array, which is actually instanced in the attribute set.
+ */
+template <typename AttributeType>
+class TMeshAttributeArraySet : public FMeshAttributeArraySetBase
+{
+	using Super = FMeshAttributeArraySetBase;
+
+public:
 	/** Constructor */
-	TAttributeIndicesArray( const int32 NumberOfIndices, const AttributeType& InDefaultValue, const EMeshAttributeFlags InFlags, const int32 InNumberOfElements )
-		: NumElements( InNumberOfElements ),
-		  DefaultValue( InDefaultValue ),
-		  Flags( InFlags )
+	FORCEINLINE explicit TMeshAttributeArraySet( const int32 NumberOfIndices = 0, const AttributeType& InDefaultValue = AttributeType(), const EMeshAttributeFlags InFlags = EMeshAttributeFlags::None, const int32 InNumberOfElements = 0 )
+		: Super( TTupleIndex<AttributeType, AttributeTypes>::Value, InFlags, InNumberOfElements ),
+		  DefaultValue( InDefaultValue )
 	{
 		SetNumIndices( NumberOfIndices );
 	}
 
+	/** Creates a copy of itself and returns a TUniquePtr to it */
+	virtual TUniquePtr<FMeshAttributeArraySetBase> Clone() const
+	{
+		return MakeUnique<TMeshAttributeArraySet>( *this );
+	}
+
 	/** Insert the element at the given index */
-	FORCEINLINE void Insert( const ElementIDType ElementID )
+	virtual void Insert( const int32 Index )
 	{
-		for( TMeshAttributeArray<AttributeType, ElementIDType>& ArrayForIndex : ArrayForIndices )
+		for( TMeshAttributeArrayBase<AttributeType>& ArrayForIndex : ArrayForIndices )
 		{
-			ArrayForIndex.Insert( ElementID, DefaultValue );
+			ArrayForIndex.Insert( Index, DefaultValue );
 		}
 
-		NumElements = FMath::Max( NumElements, ElementID.GetValue() + 1 );
+		NumElements = FMath::Max( NumElements, Index + 1 );
 	}
 
 	/** Remove the element at the given index, replacing it with a default value */
-	FORCEINLINE void Remove( const ElementIDType ElementID )
+	virtual void Remove( const int32 Index )
 	{
-		for( TMeshAttributeArray<AttributeType, ElementIDType>& ArrayForIndex : ArrayForIndices )
+		for( TMeshAttributeArrayBase<AttributeType>& ArrayForIndex : ArrayForIndices )
 		{
-			ArrayForIndex[ ElementID ] = DefaultValue;
+			ArrayForIndex[ Index ] = DefaultValue;
 		}
 	}
 
-	/** Return the TMeshAttributeArray corresponding to the given attribute index */
-	FORCEINLINE const TMeshAttributeArray<AttributeType, ElementIDType>& GetArrayForIndex( const int32 Index ) const { return ArrayForIndices[ Index ]; }
-	FORCEINLINE TMeshAttributeArray<AttributeType, ElementIDType>& GetArrayForIndex( const int32 Index ) { return ArrayForIndices[ Index ]; }
+	/** Sets the number of elements to the exact number provided, and initializes them to the default value */
+	virtual void Initialize( const int32 Count )
+	{
+		NumElements = Count;
+		for( TMeshAttributeArrayBase<AttributeType>& ArrayForIndex : ArrayForIndices )
+		{
+			ArrayForIndex.Initialize( Count, DefaultValue );
+		}
+	}
 
-	/** Return flags for this attribute type */
-	FORCEINLINE EMeshAttributeFlags GetFlags() const { return Flags; }
+	/** Polymorphic serialization */
+	virtual void Serialize( FArchive& Ar )
+	{
+		Ar << ( *this );
+	}
 
-	/** Return default value for this attribute type */
-	FORCEINLINE AttributeType GetDefaultValue() const { return DefaultValue; }
+	/** Performs an element index remap according to the passed array */
+	virtual void Remap( const TSparseArray<int32>& IndexRemap )
+	{
+		for( TMeshAttributeArrayBase<AttributeType>& ArrayForIndex : ArrayForIndices )
+		{
+			ArrayForIndex.Remap( IndexRemap, DefaultValue );
+			NumElements = ArrayForIndex.Num();
+		}
+	}
 
 	/** Return number of indices this attribute has */
-	FORCEINLINE int32 GetNumIndices() const { return ArrayForIndices.Num(); }
-
-	/** Return number of elements each attribute index has */
-	FORCEINLINE int32 GetNumElements() const { return NumElements; }
+	virtual inline int32 GetNumIndices() const { return ArrayForIndices.Num(); }
 
 	/** Sets number of indices this attribute has */
-	void SetNumIndices( const int32 NumIndices )
+	virtual void SetNumIndices( const int32 NumIndices )
 	{
-		check( NumIndices > 0 );
 		const int32 OriginalNumIndices = ArrayForIndices.Num();
 		ArrayForIndices.SetNum( NumIndices );
 
@@ -292,331 +367,334 @@ public:
 		}
 	}
 
-	/** Sets the number of elements to the exact number provided, and initializes them to the default value */
-	void Initialize( const int32 Count )
-	{
-		NumElements = Count;
-		for( TMeshAttributeArray<AttributeType, ElementIDType>& ArrayForIndex : ArrayForIndices )
-		{
-			ArrayForIndex.Initialize( Count, DefaultValue );
-		}
-	}
+	/** Return the TMeshAttributeArrayBase corresponding to the given attribute index */
+	FORCEINLINE const TMeshAttributeArrayBase<AttributeType>& GetArrayForIndex( const int32 Index ) const { return ArrayForIndices[ Index ]; }
+	FORCEINLINE TMeshAttributeArrayBase<AttributeType>& GetArrayForIndex( const int32 Index ) { return ArrayForIndices[ Index ]; }
 
-	/** Remaps all attribute indices according to the passed mapping */
-	void Remap( const TSparseArray<ElementIDType>& IndexRemap )
-	{
-		for( TMeshAttributeArray<AttributeType, ElementIDType>& ArrayForIndex : ArrayForIndices )
-		{
-			ArrayForIndex.Remap( IndexRemap, DefaultValue );
-			NumElements = ArrayForIndex.Num();
-		}
-	}
+	/** Return default value for this attribute type */
+	FORCEINLINE AttributeType GetDefaultValue() const { return DefaultValue; }
 
 	/** Serializer */
-	friend FArchive& operator<<( FArchive& Ar, TAttributeIndicesArray& AttributesArray )
+	friend FArchive& operator<<( FArchive& Ar, TMeshAttributeArraySet& AttributeArraySet )
 	{
-		Ar << AttributesArray.NumElements;
-		Ar << AttributesArray.ArrayForIndices;
-		Ar << AttributesArray.DefaultValue;
-		Ar << AttributesArray.Flags;
+		Ar << AttributeArraySet.NumElements;
+		Ar << AttributeArraySet.ArrayForIndices;
+		Ar << AttributeArraySet.DefaultValue;
+		Ar << AttributeArraySet.Flags;
 		return Ar;
 	}
 
-private:
-	/** Number of elements in each index */
-	int32 NumElements;
-
+protected:
 	/** An array of MeshAttributeArrays, one per attribute index */
-	TArray<TMeshAttributeArray<AttributeType, ElementIDType>> ArrayForIndices;
+	TArray<TMeshAttributeArrayBase<AttributeType>> ArrayForIndices;
 
 	/** The default value for an attribute of this name */
 	AttributeType DefaultValue;
-
-	/** Implementation-defined attribute name flags */
-	EMeshAttributeFlags Flags;
 };
 
 
-/** Define aliases for element attributes */
-template <typename AttributeType> using TVertexAttributeIndicesArray = TAttributeIndicesArray<AttributeType, FVertexID>;
-template <typename AttributeType> using TVertexInstanceAttributeIndicesArray = TAttributeIndicesArray<AttributeType, FVertexInstanceID>;
-template <typename AttributeType> using TEdgeAttributeIndicesArray = TAttributeIndicesArray<AttributeType, FEdgeID>;
-template <typename AttributeType> using TPolygonAttributeIndicesArray = TAttributeIndicesArray<AttributeType, FPolygonID>;
-template <typename AttributeType> using TPolygonGroupAttributeIndicesArray = TAttributeIndicesArray<AttributeType, FPolygonGroupID>;
-
-
-/**
- * This maps an attribute name to a TAttributeIndicesArray, i.e. an array of MeshAttributeArrays, one per attribute index.
- */
+// This is a deprecated class which will be removed in 4.21
 template <typename T, typename U>
-class TAttributesMap
+class TAttributeIndicesArray : public TMeshAttributeArraySet<T>
 {
 public:
+
 	using AttributeType = T;
 	using ElementIDType = U;
-	using AttributeIndicesArrayType = TAttributeIndicesArray<AttributeType, ElementIDType>;
-	using MapType = TMap<FName, AttributeIndicesArrayType>;
-
-	FORCEINLINE TAttributesMap()
-		: NumElements( 0 )
-	{}
 
-	/** Register an attribute name */
-	FORCEINLINE void RegisterAttribute( const FName AttributeName, const int32 NumberOfIndices, const AttributeType& Default, const EMeshAttributeFlags Flags )
+	/** Return the TMeshAttributeArray corresponding to the given attribute index */
+	DEPRECATED( 4.20, "Please use GetAttributesRef() or GetAttributesView() instead." )
+	FORCEINLINE const TMeshAttributeArray<AttributeType, ElementIDType>& GetArrayForIndex( const int32 Index ) const
 	{
-		if( !Map.Contains( AttributeName ) )
-		{
-			Map.Emplace( AttributeName, TAttributeIndicesArray<AttributeType, ElementIDType>( NumberOfIndices, Default, Flags, NumElements ) );
-		}
+		return static_cast<const TMeshAttributeArray<AttributeType, ElementIDType>&>( this->ArrayForIndices[ Index ] );
 	}
 
-	/** Unregister an attribute name */
-	FORCEINLINE void UnregisterAttribute( const FName AttributeName )
+	DEPRECATED( 4.20, "Please use GetAttributesRef() or GetAttributesView() instead." )
+	FORCEINLINE TMeshAttributeArray<AttributeType, ElementIDType>& GetArrayForIndex( const int32 Index )
 	{
-		Map.Remove( AttributeName );
+		return static_cast<TMeshAttributeArray<AttributeType, ElementIDType>&>( this->ArrayForIndices[ Index ] );
 	}
+};
 
-	/** Determines whether an attribute exists with the given name */
-	FORCEINLINE bool HasAttribute( const FName AttributeName ) const
-	{
-		return Map.Contains( AttributeName );
-	}
 
-	/** Get attribute array with the given name and index */
-	FORCEINLINE TMeshAttributeArray<AttributeType, ElementIDType>& GetAttributes( const FName AttributeName, const int32 AttributeIndex = 0 )
-	{
-		// @todo mesh description: should this handle non-existent attribute names and indices gracefully?
-		return Map.FindChecked( AttributeName ).GetArrayForIndex( AttributeIndex );
-	}
+/**
+ * This is the class used to access attribute values.
+ * It is a proxy object to a TMeshAttributeArraySet<> and should be passed by value.
+ * It is valid for as long as the owning FMeshDescription exists.
+ * Note that this only provides non-mutating accessors; a mutating version is derived from this.
+ */
+template <typename ElementIDType, typename AttributeType>
+class TMeshAttributesConstRef
+{
+public:
+	using Type = AttributeType;
+	using ArrayType = TMeshAttributeArraySet<AttributeType>;
 
-	FORCEINLINE const TMeshAttributeArray<AttributeType, ElementIDType>& GetAttributes( const FName AttributeName, const int32 AttributeIndex = 0 ) const
-	{
-		// @todo mesh description: should this handle non-existent attribute names and indices gracefully?
-		return Map.FindChecked( AttributeName ).GetArrayForIndex( AttributeIndex );
-	}
+	FORCEINLINE explicit TMeshAttributesConstRef( const ArrayType* InArrayPtr = nullptr )
+		: ArrayPtr( InArrayPtr )
+	{}
 
-	/** Get attribute indices array with the given name */
-	FORCEINLINE TAttributeIndicesArray<AttributeType, ElementIDType>& GetAttributesSet( const FName AttributeName )
+	/** Access elements from attribute index 0 */
+	FORCEINLINE const AttributeType& operator[]( const ElementIDType ElementID ) const
 	{
-		// @todo mesh description: should this handle non-existent attribute names gracefully?
-		return Map.FindChecked( AttributeName );
+		return ArrayPtr->GetArrayForIndex( 0 )[ ElementID.GetValue() ];
 	}
 
-	FORCEINLINE const TAttributeIndicesArray<AttributeType, ElementIDType>& GetAttributesSet( const FName AttributeName ) const
+	/** Get the element with the given ID from index 0 */
+	FORCEINLINE AttributeType Get( const ElementIDType ElementID ) const
 	{
-		// @todo mesh description: should this handle non-existent attribute names gracefully?
-		return Map.FindChecked( AttributeName );
+		return ArrayPtr->GetArrayForIndex( 0 )[ ElementID.GetValue() ];
 	}
 
-	/** Returns the number of indices for the attribute with the given name */
-	FORCEINLINE int32 GetAttributeIndexCount( const FName AttributeName ) const
+	/** Get the element with the given ID and index */
+	FORCEINLINE AttributeType Get( const ElementIDType ElementID, const int32 Index ) const
 	{
-		// @todo mesh description: should this handle non-existent attribute names and indices gracefully?
-		return Map.FindChecked( AttributeName ).GetNumIndices();
+		return ArrayPtr->GetArrayForIndex( Index )[ ElementID.GetValue() ];
 	}
 
-	/** Sets the number of indices for the attribute with the given name */
-	FORCEINLINE void SetAttributeIndexCount( const FName AttributeName, const int32 NumIndices )
-	{
-		Map.FindChecked( AttributeName ).SetNumIndices( NumIndices );
-	}
+	/** Return whether the reference is valid or not */
+	FORCEINLINE bool IsValid() const { return ( ArrayPtr != nullptr ); }
 
-	/** Returns an array of all the attribute names registered for this attribute type */
-	template <typename Allocator>
-	FORCEINLINE void GetAttributeNames( TArray<FName, Allocator>& OutAttributeNames ) const
+	/** Return default value for this attribute type */
+	FORCEINLINE AttributeType GetDefaultValue() const { return ArrayPtr->GetDefaultValue(); }
+
+	/** Return number of indices this attribute has */
+	FORCEINLINE int32 GetNumIndices() const
 	{
-		Map.GetKeys( OutAttributeNames );
+		return ArrayPtr->ArrayType::GetNumIndices();	// note: override virtual dispatch
 	}
 
-	/** Gets a single attribute with the given ElementID, Name and Index */
-	FORCEINLINE AttributeType GetAttribute( const ElementIDType ElementID, const FName AttributeName, const int32 AttributeIndex = 0 ) const
+	/** Get the number of elements in this attribute array */
+	FORCEINLINE int32 GetNumElements() const
 	{
-		return Map.FindChecked( AttributeName ).GetArrayForIndex( AttributeIndex )[ ElementID ];
+		return ArrayPtr->GetNumElements();
 	}
 
-	/** Sets a single attribute with the given ElementID, Name and Index to the given value */
-	FORCEINLINE void SetAttribute( const ElementIDType ElementID, const FName AttributeName, const int32 AttributeIndex, const AttributeType& AttributeValue )
+	/** Get the flags for this attribute array set */
+	FORCEINLINE EMeshAttributeFlags GetFlags() const { return ArrayPtr->GetFlags(); }
+
+protected:
+	const ArrayType* ArrayPtr;
+};
+
+
+/**
+ * This is a version which provides mutating accessors.
+ * This hierarchy allows us to assign MeshAttributesConstRef = MeshAttributesRef.
+ */
+template <typename ElementIDType, typename AttributeType>
+class TMeshAttributesRef final : public TMeshAttributesConstRef<ElementIDType, AttributeType>
+{
+public:
+	using Super = TMeshAttributesConstRef<ElementIDType, AttributeType>;
+	using Type = typename Super::Type;
+	using ArrayType = typename Super::ArrayType;
+
+	FORCEINLINE explicit TMeshAttributesRef( const ArrayType* InArrayPtr = nullptr )
+		: Super( InArrayPtr )
+	{}
+
+	/** Access elements from attribute index 0 */
+	FORCEINLINE AttributeType& operator[]( const ElementIDType ElementID ) const
 	{
-		Map.FindChecked( AttributeName ).GetArrayForIndex( AttributeIndex )[ ElementID ] = AttributeValue;
+		return const_cast<ArrayType*>( this->ArrayPtr )->GetArrayForIndex( 0 )[ ElementID.GetValue() ];
 	}
 
-	/** Inserts a default-initialized value for all attributes of the given ID */
-	void Insert( const ElementIDType ElementID )
+	/** Set the element with the given ID and index 0 to the provided value */
+	FORCEINLINE void Set( const ElementIDType ElementID, const AttributeType& Value ) const
 	{
-		NumElements = FMath::Max( NumElements, ElementID.GetValue() + 1 );
-		for( auto& AttributeNameAndIndicesArray : Map )
-		{
-			AttributeNameAndIndicesArray.Value.Insert( ElementID );
-			check( AttributeNameAndIndicesArray.Value.GetNumElements() == this->NumElements );
-		}
+		const_cast<ArrayType*>( this->ArrayPtr )->GetArrayForIndex( 0 )[ ElementID.GetValue() ] = Value;
 	}
 
-	/** Removes all attributes with the given ID */
-	void Remove( const ElementIDType ElementID )
+	/** Set the element with the given ID and index to the provided value */
+	FORCEINLINE void Set( const ElementIDType ElementID, const int32 Index, const AttributeType& Value ) const
 	{
-		for( auto& AttributeNameAndIndicesArray : Map )
-		{
-			AttributeNameAndIndicesArray.Value.Remove( ElementID );
-		}
+		const_cast<ArrayType*>( this->ArrayPtr )->GetArrayForIndex( Index )[ ElementID.GetValue() ] = Value;
 	}
 
-	/** Initializes all attributes to have the given number of elements with the default value */
-	void Initialize( const int32 Count )
+	/** Sets number of indices this attribute has */
+	FORCEINLINE void SetNumIndices( const int32 NumIndices )
 	{
-		NumElements = Count;
-		for( auto& AttributeNameAndIndicesArray : Map )
-		{
-			AttributeNameAndIndicesArray.Value.Initialize( Count );
-		}
+		const_cast<ArrayType*>( this->ArrayPtr )->ArrayType::SetNumIndices( NumIndices );	// note: override virtual dispatch
 	}
+};
+
+
+
+/**
+ * This is the class used to provide a 'view' of the specified type on an attribute array.
+ * Like TMeshAttributesRef, it is a proxy object which is valid for as long as the owning FMeshDescription exists,
+ * and should be passed by value.
+ *
+ * This is the base class, and shouldn't be instanced directly.
+ */
+template <typename ViewType>
+class TMeshAttributesViewBase
+{
+public:
+	/** Return whether the reference is valid or not */
+	FORCEINLINE bool IsValid() const { return ( ArrayPtr != nullptr ); }
+
+	/** Return number of indices this attribute has */
+	FORCEINLINE int32 GetNumIndices() const { return ArrayPtr->GetNumIndices(); }
+
+	/** Return default value for this attribute type */
+	FORCEINLINE ViewType GetDefaultValue() const;
 
-	/** Returns the number of elements held by each attribute in this map */
+	/** Get the number of elements in this attribute array */
 	FORCEINLINE int32 GetNumElements() const
 	{
-		return NumElements;
+		return ArrayPtr->GetNumElements();
 	}
 
-	/**
-	 * Call the supplied function on each attribute.
-	 * The prototype should be Func( const FName AttributeName, auto& AttributeIndicesArray );
-	 */
-	template <typename FuncType>
-	void ForEachAttributeIndicesArray( const FuncType& Func )
-	{
-		for( auto& AttributeNameAndIndicesArray : Map )
-		{
-			Func( AttributeNameAndIndicesArray.Key, AttributeNameAndIndicesArray.Value );
-		}
-	}
+protected:
+	FORCEINLINE explicit TMeshAttributesViewBase( const FMeshAttributeArraySetBase* InArrayPtr )
+		: ArrayPtr( InArrayPtr )
+	{}
 
-	template <typename FuncType>
-	void ForEachAttributeIndicesArray( const FuncType& Func ) const
-	{
-		for( const auto& AttributeNameAndIndicesArray : Map )
-		{
-			Func( AttributeNameAndIndicesArray.Key, AttributeNameAndIndicesArray.Value );
-		}
-	}
+	/** Get the element with the given ID from index 0 */
+	FORCEINLINE ViewType GetByIndex( const int32 ElementIndex ) const;
 
-	/** Registers attributes copied from the specified attributes map */
-	void RegisterAttributesFromAttributesSet( const TAttributesMap& SrcAttributesMap )
-	{
-		for( const auto& AttributeNameAndIndicesArray : SrcAttributesMap.Map )
-		{
-			const FName& AttributeName = AttributeNameAndIndicesArray.Key;
-			const AttributeIndicesArrayType& AttributeIndicesArray = AttributeNameAndIndicesArray.Value;
+	/** Get the element with the given element and attribute indices */
+	FORCEINLINE ViewType GetByIndex( const int32 ElementIndex, const int32 AttributeIndex ) const;
 
-			RegisterAttribute( AttributeName, AttributeIndicesArray.GetNumIndices(), AttributeIndicesArray.GetDefaultValue(), AttributeIndicesArray.GetFlags() );
-		}
-	}
+	/** Set the attribute index 0 element with the given index to the provided value */
+	FORCEINLINE void SetByIndex( const int32 ElementIndex, const ViewType& Value ) const;
 
-	/** Applies the given remapping to the attributes set */
-	void Remap( const TSparseArray<ElementIDType>& IndexRemap )
-	{
-		if( Map.Num() == 0 )
-		{
-			// If there are no attributes registered, determine the number of elements by finding the maximum
-			// remapped element ID in the IndexRemap array
-			NumElements = 0;
-			for( const ElementIDType ElementID : IndexRemap )
-			{
-				NumElements = FMath::Max( NumElements, ElementID.GetValue() + 1 );
-			}
-		}
-		else
-		{
-			// Otherwise perform the remap, and get the number of elements from the resulting attribute indices array.
-			for( auto& AttributeNameAndIndicesArray : Map )
-			{
-				AttributeNameAndIndicesArray.Value.Remap( IndexRemap );
-				NumElements = AttributeNameAndIndicesArray.Value.GetNumElements();
-			}
-		}
-	}
+	/** Set the element with the given element and attribute indices to the provided value */
+	FORCEINLINE void SetByIndex( const int32 ElementIndex, const int32 AttributeIndex, const ViewType& Value ) const;
 
-private:
-	friend FArchive& operator<<( FArchive& Ar, TAttributesMap& AttributesMap )
-	{
-		// First serialize the number of elements which each attribute should contain
-		Ar << AttributesMap.NumElements;
+	const FMeshAttributeArraySetBase* ArrayPtr;
+};
 
-		// Now serialize the attributes of this type.
-		Ar << AttributesMap.Map;
 
-		return Ar;
-	}
+/**
+ * This is a derived version with typesafe element accessors, which is returned by TAttributesSet<>.
+ * It is also limited to non-mutating accessors, and is returned by GetAttributesRef on a const attribute set.
+ */
+template <typename ElementIDType, typename ViewType>
+class TMeshAttributesConstView : public TMeshAttributesViewBase<ViewType>
+{
+	using Super = TMeshAttributesViewBase<ViewType>;
 
-	/** Number of elements for each attribute index */
-	int32 NumElements;
+public:
+	FORCEINLINE explicit TMeshAttributesConstView( const FMeshAttributeArraySetBase* InArrayPtr = nullptr )
+		: Super( InArrayPtr )
+	{}
 
-	/** The actual container */
-	MapType Map;
+	/** Get the element with the given ID from index 0. This version has a typesafe element ID accessor. */
+	FORCEINLINE ViewType Get( const ElementIDType ElementID ) const { return this->GetByIndex( ElementID.GetValue() ); }
+
+	/** Get the element with the given ID and index. This version has a typesafe element ID accessor. */
+	FORCEINLINE ViewType Get( const ElementIDType ElementID, const int32 Index ) const { return this->GetByIndex( ElementID.GetValue(), Index ); }
 };
 
 
 /**
- * Helper template which transforms a tuple of types into a tuple of TAttributesArrays of those types.
- *
- * We need to instance TAttributeArrays for each type in the AttributeTypes tuple.
- * Then we can access the appropriate array (as long as we know what its index is).
- *
- * This template, given ElementIDType and TTuple<A, B>, will generate:
- * TTuple<TAttributesArray<A, ElementIDType>, TAttributesArray<B, ElementIDType>>
+ * This is a derived version with mutating accessors.
+ * This type of hierarchy means it's possible to assign a MeshAttributesConstView = MeshAttributesView.
  */
-template <typename ElementIDType, typename Tuple>
-struct TMakeAttributesSet;
-
-template <typename ElementIDType, typename... TupleTypes>
-struct TMakeAttributesSet<ElementIDType, TTuple<TupleTypes...>>
+template <typename ElementIDType, typename ViewType>
+class TMeshAttributesView final : public TMeshAttributesConstView<ElementIDType, ViewType>
 {
-	using Type = TTuple<TAttributesMap<TupleTypes, ElementIDType>...>;
+	using Super = TMeshAttributesConstView<ElementIDType, ViewType>;
+
+public:
+	FORCEINLINE explicit TMeshAttributesView( FMeshAttributeArraySetBase* InArrayPtr = nullptr )
+		: Super( InArrayPtr )
+	{}
+
+	/** Set the element with the given ID and index 0 to the provided value. This version has a typesafe element ID accessor. */
+	FORCEINLINE void Set( const ElementIDType ElementID, const ViewType& Value ) const { this->SetByIndex( ElementID.GetValue(), Value ); }
+
+	/** Set the element with the given ID and index to the provided value. This version has a typesafe element ID accessor. */
+	FORCEINLINE void Set( const ElementIDType ElementID, const int32 Index, const ViewType& Value ) const { this->SetByIndex( ElementID.GetValue(), Index, Value ); }
+
+	/** Sets number of indices this attribute has */
+	FORCEINLINE void SetNumIndices( const int32 NumIndices ) const
+	{
+		const_cast<FMeshAttributeArraySetBase*>( this->ArrayPtr )->SetNumIndices( NumIndices );
+	}
 };
 
 
 /**
- * Helper template which gets the tuple index of a given type from a given TTuple.
- *
- * Given Type = char, and Tuple = TTuple<int, float, char>,
- * TTupleIndex<Type, Tuple>::Value will be 2.
+ * This is a wrapper for an allocated attributes array.
+ * It holds a TUniquePtr pointing to the actual attributes array, and performs polymorphic copy and assignment,
+ * as per the actual array type.
  */
-template <typename Type, typename Tuple>
-struct TTupleIndex;
-
-template <typename Type, typename... Types>
-struct TTupleIndex<Type, TTuple<Type, Types...>>
+class FAttributesSetEntry
 {
-	static const uint32 Value = 0U;
-};
+public:
+	/**
+	 * Default constructor.
+	 * This breaks the invariant that Ptr be always valid, but is necessary so that it can be the value type of a TMap.
+	 */
+	FORCEINLINE FAttributesSetEntry() = default;
 
-template <typename Type, typename Head, typename... Tail>
-struct TTupleIndex<Type, TTuple<Head, Tail...>>
-{
-	static const uint32 Value = 1U + TTupleIndex<Type, TTuple<Tail...>>::Value;
-};
+	/**
+	 * Construct a valid FAttributesSetEntry of the concrete type specified.
+	 */
+	template <typename AttributeType>
+	FORCEINLINE FAttributesSetEntry( const int32 NumberOfIndices, const AttributeType& Default, const EMeshAttributeFlags Flags, const int32 NumElements )
+		: Ptr( MakeUnique<TMeshAttributeArraySet<AttributeType>>( NumberOfIndices, Default, Flags, NumElements ) )
+	{}
 
+	/** Default destructor */
+	FORCEINLINE ~FAttributesSetEntry() = default;
 
-/** Helper template which splits a tuple into head and tail */
-template <typename Tuple>
-struct TSplitTuple;
+	/** Polymorphic copy: a new copy of Other is created */
+	FAttributesSetEntry( const FAttributesSetEntry& Other )
+		: Ptr( Other.Ptr ? Other.Ptr->Clone() : nullptr )
+	{}
 
-template <typename TupleHead, typename... TupleTail>
-struct TSplitTuple<TTuple<TupleHead, TupleTail...>>
-{
-	using Head = TupleHead;
-	using Tail = TTuple<TupleTail...>;
+	/** Default move constructor */
+	FAttributesSetEntry( FAttributesSetEntry&& ) = default;
+
+	/** Polymorphic assignment */
+	FAttributesSetEntry& operator=( const FAttributesSetEntry& Other )
+	{
+		FAttributesSetEntry Temp( Other );
+		Swap( *this, Temp );
+		return *this;
+	}
+
+	/** Default move assignment */
+	FAttributesSetEntry& operator=( FAttributesSetEntry&& ) = default;
+
+	/** Transparent access through the TUniquePtr */
+	FORCEINLINE const FMeshAttributeArraySetBase* Get() const { return Ptr.Get(); }
+	FORCEINLINE const FMeshAttributeArraySetBase* operator->() const { return Ptr.Get(); }
+	FORCEINLINE const FMeshAttributeArraySetBase& operator*() const { return *Ptr; }
+	FORCEINLINE FMeshAttributeArraySetBase* Get() { return Ptr.Get(); }
+	FORCEINLINE FMeshAttributeArraySetBase* operator->() { return Ptr.Get(); }
+	FORCEINLINE FMeshAttributeArraySetBase& operator*() { return *Ptr; }
+
+	/** Object can be coerced to bool to indicate if it is valid */
+	FORCEINLINE explicit operator bool() const { return Ptr.IsValid(); }
+	FORCEINLINE bool operator!() const { return !Ptr.IsValid(); }
+
+	/** Given a type at runtime, allocate an attribute array of that type, owned by Ptr */
+	void CreateArrayOfType( const uint32 Type );
+
+	/** Serialization */
+	friend FArchive& operator<<( FArchive& Ar, FAttributesSetEntry& Entry );
+
+private:
+	TUniquePtr<FMeshAttributeArraySetBase> Ptr;
 };
 
 
 /**
- * This is the container for all attributes of a particular mesh element.
- * It contains a TTuple of TAttributesMap, one per attribute type (FVector, float, bool, etc)
+ * This is the container for all attributes and their arrays. It wraps a TMap, mapping from attribute name to attribute array.
+ * An attribute may be of any arbitrary type; we use a mixture of polymorphism and compile-time templates to handle the different types.
  */
-template <typename ElementIDType>
-class TAttributesSet
+class FAttributesSetBase
 {
 public:
 	/**
 	 * Register a new attribute name with the given type (must be a member of the AttributeTypes tuple).
+	 * If the attribute name is already registered, it will do nothing.
 	 *
 	 * Example of use:
 	 *
@@ -628,210 +706,697 @@ public:
 	template <typename AttributeType>
 	void RegisterAttribute( const FName AttributeName, const int32 NumberOfIndices = 1, const AttributeType& Default = AttributeType(), const EMeshAttributeFlags Flags = EMeshAttributeFlags::None )
 	{
-		Container.template Get<TTupleIndex<AttributeType, AttributeTypes>::Value>().RegisterAttribute( AttributeName, NumberOfIndices, Default, Flags );
+		if( !Map.Contains( AttributeName ) )
+		{
+			Map.Emplace( AttributeName, FAttributesSetEntry( NumberOfIndices, Default, Flags, NumElements ) );
+		}
+	}
+
+	/**
+	 * Unregister an attribute with the given name.
+	 */
+	void UnregisterAttribute( const FName AttributeName )
+	{
+		Map.Remove( AttributeName );
+	}
+
+	/** Determines whether an attribute exists with the given name */
+	bool HasAttribute( const FName AttributeName ) const
+	{
+		return ( Map.Contains( AttributeName ) );
+	}
+
+	/**
+	 * Determines whether an attribute of the given type exists with the given name
+	 */
+	template <typename AttributeType>
+	bool HasAttributeOfType( const FName AttributeName ) const
+	{
+		if( const FAttributesSetEntry* ArraySetPtr = Map.Find( AttributeName ) )
+		{
+			return ( *ArraySetPtr )->HasType<AttributeType>();
+		}
+
+		return false;
+	}
+
+	/** Initializes all attributes to have the given number of elements with the default value */
+	void Initialize( const int32 Count )
+	{
+		NumElements = Count;
+		for( auto& MapEntry : Map )
+		{
+			MapEntry.Value->Initialize( Count );
+		}
 	}
 
-	/** Unregister an attribute name with the given type */
+	/** Applies the given remapping to the attributes set */
+	void Remap( const TSparseArray<int32>& IndexRemap );
+
 	template <typename AttributeType>
+	DEPRECATED( 4.20, "Please use untemplated UnregisterAttribute() instead" )
 	void UnregisterAttribute( const FName AttributeName )
 	{
-		Container.template Get<TTupleIndex<AttributeType, AttributeTypes>::Value>().UnregisterAttribute( AttributeName );
+		return UnregisterAttribute( AttributeName );
 	}
 
-	/** Determines whether an attribute of the given type exists with the given name */
 	template <typename AttributeType>
-	bool HasAttribute( const FName AttributeName ) const
+	DEPRECATED( 4.20, "Please use untemplated HasAttribute() instead" )
+	void HasAttribute( const FName AttributeName )
 	{
-		return Container.template Get<TTupleIndex<AttributeType, AttributeTypes>::Value>().HasAttribute( AttributeName );
+		return HasAttribute( AttributeName );
 	}
 
+protected:
+	/**
+	 * Insert a new element at the given index.
+	 * The public API version of this function takes an ID of ElementIDType instead of a typeless index.
+	 */
+	void Insert( const int32 Index )
+	{
+		NumElements = FMath::Max( NumElements, Index + 1 );
+
+		for( auto& MapEntry : Map )
+		{
+			MapEntry.Value->Insert( Index );
+			check( MapEntry.Value->GetNumElements() == NumElements );
+		}
+	}
+
+	/**
+	 * Remove an element at the given index.
+	 * The public API version of this function takes an ID of ElementIDType instead of a typeless index.
+	 */
+	void Remove( const int32 Index )
+	{
+		for( auto& MapEntry : Map )
+		{
+			MapEntry.Value->Remove( Index );
+		}
+	}
+
+	/** Serialization */
+	friend FArchive& operator<<( FArchive& Ar, FAttributesSetBase& AttributesSet );
+
+	template <typename T>
+	friend void SerializeLegacy( FArchive& Ar, FAttributesSetBase& AttributesSet );
+
+	/** The actual container */
+	TMap<FName, FAttributesSetEntry> Map;
+
+	/** The number of elements in each attribute array */
+	int32 NumElements;
+};
+
+
+/**
+ * This is a version of the attributes set container which accesses elements by typesafe IDs.
+ * This prevents access of (for example) vertex instance attributes by vertex IDs.
+ */
+template <typename ElementIDType>
+class TAttributesSet final : public FAttributesSetBase
+{
+	using FAttributesSetBase::Insert;
+	using FAttributesSetBase::Remove;
+
+public:
 	/**
-	 * Get an attribute array with the given type, name and index.
+	 * Get an attribute array with the given type and name.
+	 * The attribute type must correspond to the type passed as the template parameter.
 	 *
 	 * Example of use:
 	 *
-	 *		const TVertexAttributeArray<FVector>& VertexPositions = VertexAttributes().GetAttributes<FVector>( "Position" );
+	 *		TVertexAttributesConstRef<FVector> VertexPositions = VertexAttributes().GetAttributesRef<FVector>( "Position" ); // note: assign to value type
 	 *		for( const FVertexID VertexID : GetVertices().GetElementIDs() )
 	 *		{
-	 *			const FVector Position = VertexPositions[ VertexID ];
+	 *			const FVector Position = VertexPositions.Get( VertexID );
 	 *			DoSomethingWith( Position );
 	 *		}
+	 *
+	 * Note that the returned object is a value type which should be assigned and passed by value, not reference.
+	 * It is valid for as long as this TAttributesSet object exists.
 	 */
 	template <typename AttributeType>
-	TMeshAttributeArray<AttributeType, ElementIDType>& GetAttributes( const FName AttributeName, const int32 AttributeIndex = 0 )
+	TMeshAttributesConstRef<ElementIDType, AttributeType> GetAttributesRef( const FName AttributeName ) const
 	{
-		// @todo mesh description: should this handle non-existent attribute names and indices gracefully?
-		return Container.template Get<TTupleIndex<AttributeType, AttributeTypes>::Value>().GetAttributes( AttributeName, AttributeIndex );
+		if( const FAttributesSetEntry* ArraySetPtr = this->Map.Find( AttributeName ) )
+		{
+			if( ( *ArraySetPtr )->HasType<AttributeType>() )
+			{
+				return TMeshAttributesConstRef<ElementIDType, AttributeType>( static_cast<const TMeshAttributeArraySet<AttributeType>*>( ArraySetPtr->Get() ) );
+			}
+		}
+
+		return TMeshAttributesConstRef<ElementIDType, AttributeType>();
 	}
 
 	template <typename AttributeType>
-	const TMeshAttributeArray<AttributeType, ElementIDType>& GetAttributes( const FName AttributeName, const int32 AttributeIndex = 0 ) const
+	TMeshAttributesRef<ElementIDType, AttributeType> GetAttributesRef( const FName AttributeName )
 	{
-		// @todo mesh description: should this handle non-existent attribute names and indices gracefully?
-		return Container.template Get<TTupleIndex<AttributeType, AttributeTypes>::Value>().GetAttributes( AttributeName, AttributeIndex );
+		if( FAttributesSetEntry* ArraySetPtr = this->Map.Find( AttributeName ) )
+		{
+			if( ( *ArraySetPtr )->HasType<AttributeType>() )
+			{
+				return TMeshAttributesRef<ElementIDType, AttributeType>( static_cast<TMeshAttributeArraySet<AttributeType>*>( ArraySetPtr->Get() ) );
+			}
+		}
+
+		return TMeshAttributesRef<ElementIDType, AttributeType>();
 	}
 
 	/**
-	 * Get a set of attribute arrays with the given type and name.
+	 * Get a view on an attribute array with the given name, accessing elements as the given type.
+	 * Access to elements will be slightly slower than with GetAttributesRef, but element access is not strongly typed.
 	 *
 	 * Example of use:
 	 *
-	 *		const TArray<TVertexInstanceAttributeArray<FVector2D>>& UVs = VertexInstanceAttributes().GetAttributesSet<FVector2D>( "UV" );
+	 *		const TVertexInstanceAttributesView<FVector> VertexNormals = VertexInstanceAttributes().GetAttributesView<FVector>( "Normal" );
 	 *		for( const FVertexInstanceID VertexInstanceID : GetVertexInstances().GetElementIDs() )
 	 *		{
-	 *			const FVector2D UV0 = UVs[ 0 ][ VertexInstanceID ];
-	 *			const FVector2D UV1 = UVs[ 1 ][ VertexInstanceID ];
-	 *			DoSomethingWith( UV0, UV1 );
+	 *          // This will work even if the Normals array has a different internal type, e.g. FPackedVector
+	 *			const FVector Normal = VertexNormals.Get( VertexInstanceID );
+	 *			DoSomethingWith( Normal );
 	 *		}
+	 *
+	 * Note that the returned object is a value type which should be assigned and passed by value, not reference.
+	 * It is valid for as long as this TAttributesSet object exists.
 	 */
-	template <typename AttributeType>
-	TAttributeIndicesArray<AttributeType, ElementIDType>& GetAttributesSet( const FName AttributeName )
+	template <typename ViewType>
+	TMeshAttributesConstView<ElementIDType, ViewType> GetAttributesView( const FName AttributeName ) const
 	{
-		// @todo mesh description: should this handle non-existent attribute names gracefully?
-		return Container.template Get<TTupleIndex<AttributeType, AttributeTypes>::Value>().GetAttributesSet( AttributeName );
+		if( const FAttributesSetEntry* ArraySetPtr = this->Map.Find( AttributeName ) )
+		{
+			return TMeshAttributesConstView<ElementIDType, ViewType>( ArraySetPtr->Get() );
+		}
+
+		return TMeshAttributesConstView<ElementIDType, ViewType>();
 	}
 
-	template <typename AttributeType>
-	const TAttributeIndicesArray<AttributeType, ElementIDType>& GetAttributesSet( const FName AttributeName ) const
+	template <typename ViewType>
+	TMeshAttributesView<ElementIDType, ViewType> GetAttributesView( const FName AttributeName )
 	{
-		// @todo mesh description: should this handle non-existent attribute names gracefully?
-		return Container.template Get<TTupleIndex<AttributeType, AttributeTypes>::Value>().GetAttributesSet( AttributeName );
+		if( FAttributesSetEntry* ArraySetPtr = this->Map.Find( AttributeName ) )
+		{
+			return TMeshAttributesView<ElementIDType, ViewType>( ArraySetPtr->Get() );
+		}
+
+		return TMeshAttributesView<ElementIDType, ViewType>();
 	}
 
 	/** Returns the number of indices for the attribute with the given name */
 	template <typename AttributeType>
 	int32 GetAttributeIndexCount( const FName AttributeName ) const
 	{
-		// @todo mesh description: should this handle non-existent attribute names and indices gracefully?
-		return Container.template Get<TTupleIndex<AttributeType, AttributeTypes>::Value>().GetAttributeIndexCount( AttributeName );
+		if( const FAttributesSetEntry* ArraySetPtr = this->Map.Find( AttributeName ) )
+		{
+			if( ( *ArraySetPtr )->HasType<AttributeType>() )
+			{
+				using ArrayType = TMeshAttributeArraySet<AttributeType>;
+				return static_cast<const ArrayType*>( ArraySetPtr->Get() )->ArrayType::GetNumIndices();	// note: override virtual dispatch
+			}
+		}
+
+		return 0;
 	}
 
 	/** Sets the number of indices for the attribute with the given name */
 	template <typename AttributeType>
 	void SetAttributeIndexCount( const FName AttributeName, const int32 NumIndices )
 	{
-		Container.template Get<TTupleIndex<AttributeType, AttributeTypes>::Value>().SetAttributeIndexCount( AttributeName, NumIndices );
+		if( FAttributesSetEntry* ArraySetPtr = this->Map.Find( AttributeName ) )
+		{
+			if( ( *ArraySetPtr )->HasType<AttributeType>() )
+			{
+				using ArrayType = TMeshAttributeArraySet<AttributeType>;
+				static_cast<ArrayType*>( ArraySetPtr->Get() )->ArrayType::SetNumIndices( NumIndices );	// note: override virtual dispatch
+			}
+		}
 	}
 
 	/** Returns an array of all the attribute names registered for this attribute type */
 	template <typename AttributeType, typename Allocator>
 	void GetAttributeNames( TArray<FName, Allocator>& OutAttributeNames ) const
 	{
-		Container.template Get<TTupleIndex<AttributeType, AttributeTypes>::Value>().GetAttributeNames( OutAttributeNames );
+		this->Map.GetKeys( OutAttributeNames );
 	}
 
 	template <typename AttributeType>
 	AttributeType GetAttribute( const ElementIDType ElementID, const FName AttributeName, const int32 AttributeIndex = 0 ) const
 	{
-		return Container.template Get<TTupleIndex<AttributeType, AttributeTypes>::Value>().GetAttribute( ElementID, AttributeName, AttributeIndex );
+		const FMeshAttributeArraySetBase* ArraySetPtr = this->Map.FindChecked( AttributeName ).Get();
+		check( ArraySetPtr->HasType<AttributeType>() );
+		return static_cast<const TMeshAttributeArraySet<AttributeType>*>( ArraySetPtr )->GetArrayForIndex( AttributeIndex )[ ElementID.GetValue() ];
 	}
 
 	template <typename AttributeType>
 	void SetAttribute( const ElementIDType ElementID, const FName AttributeName, const int32 AttributeIndex, const AttributeType& AttributeValue )
 	{
-		Container.template Get<TTupleIndex<AttributeType, AttributeTypes>::Value>().SetAttribute( ElementID, AttributeName, AttributeIndex, AttributeValue );
+		FMeshAttributeArraySetBase* ArraySetPtr = this->Map.FindChecked( AttributeName ).Get();
+		check( ArraySetPtr->HasType<AttributeType>() );
+		static_cast<TMeshAttributeArraySet<AttributeType>*>( ArraySetPtr )->GetArrayForIndex( AttributeIndex )[ ElementID.GetValue() ] = AttributeValue;
 	}
 
 	/** Inserts a default-initialized value for all attributes of the given ID */
-	void Insert( const ElementIDType ElementID )
+	FORCEINLINE void Insert( const ElementIDType ElementID )
 	{
-		VisitTupleElements( [ ElementID ]( auto& AttributesMap ) { AttributesMap.Insert( ElementID ); }, Container );
+		this->Insert( ElementID.GetValue() );
 	}
 
 	/** Removes all attributes with the given ID */
-	void Remove( const ElementIDType ElementID )
+	FORCEINLINE void Remove( const ElementIDType ElementID )
 	{
-		VisitTupleElements( [ ElementID ]( auto& AttributesMap ) { AttributesMap.Remove( ElementID ); }, Container );
+		this->Remove( ElementID.GetValue() );
 	}
 
-	/** Initializes the attribute set with the given number of elements, all at the default value */
-	void Initialize( const int32 NumElements )
-	{
-		VisitTupleElements( [ NumElements ]( auto& AttributesMap ) { AttributesMap.Initialize( NumElements); }, Container );
-	}
+	/**
+	 * Call the supplied function on each attribute.
+	 * The prototype should be Func( const FName AttributeName, auto AttributesRef );
+	 */
+	template <typename ForEachFunc> void ForEach( ForEachFunc Func );
+
+	/**
+	* Call the supplied function on each attribute.
+	* The prototype should be Func( const FName AttributeName, auto AttributesConstRef );
+	*/
+	template <typename ForEachFunc> void ForEach( ForEachFunc Func ) const;
 
 	/**
 	 * Call the supplied function on each attribute.
 	 * The prototype should be Func( const FName AttributeName, auto& AttributeIndicesArray );
 	 */
 	template <typename FuncType>
+	DEPRECATED( 4.20, "This is no longer supported; please use ForEach() instead and amend your lambda to accept an auto of type TMeshAttributesRef instead." )
 	void ForEachAttributeIndicesArray( const FuncType& Func )
 	{
-		VisitTupleElements( [ &Func ]( auto& AttributesMap ) { AttributesMap.ForEachAttributeIndicesArray( Func ); }, Container );
+		check( false );
 	}
 
 	template <typename FuncType>
+	DEPRECATED( 4.20, "This is no longer supported; please use ForEach() instead and amend your lambda to accept an auto of type const TMeshAttributesRef instead." )
 	void ForEachAttributeIndicesArray( const FuncType& Func ) const
 	{
-		VisitTupleElements( [ &Func ]( const auto& AttributesMap ) { AttributesMap.ForEachAttributeIndicesArray( Func ); }, Container );
+		check( false );
 	}
 
-	/** Applies the given remapping to the attributes set */
-	void Remap( const TSparseArray<ElementIDType>& IndexRemap )
+	template <typename AttributeType>
+	DEPRECATED( 4.20, "Please use GetAttributesRef() or GetAttributesView() instead." )
+	TMeshAttributeArray<AttributeType, ElementIDType>& GetAttributes( const FName AttributeName, const int32 AttributeIndex = 0 )
 	{
-		VisitTupleElements( [ &IndexRemap ]( auto& AttributesMap ) { AttributesMap.Remap( IndexRemap ); }, Container );
+		FMeshAttributeArraySetBase* ArraySetPtr = this->Map.FindChecked( AttributeName ).Get();
+		check( ArraySetPtr->HasType<AttributeType>() );
+		return static_cast<TAttributeIndicesArray<AttributeType, ElementIDType>*>( ArraySetPtr )->GetArrayForIndex( AttributeIndex );
 	}
 
-	/** Copies registered attributes from another TAttributesSet */
-	void RegisterAttributesFromAttributesSet( const TAttributesSet& Other )
+	template <typename AttributeType>
+	DEPRECATED( 4.20, "Please use GetAttributesRef() or GetAttributesView() instead." )
+	const TMeshAttributeArray<AttributeType, ElementIDType>& GetAttributes( const FName AttributeName, const int32 AttributeIndex = 0 ) const
 	{
-		VisitTupleElements( [ this ]( auto& DestAttributesMap, const auto& SrcAttributesMap ) { DestAttributesMap.RegisterAttributesFromAttributesSet( SrcAttributesMap ); },
-			Container, Other.Container );
+		const FMeshAttributeArraySetBase* ArraySetPtr = this->Map.FindChecked( AttributeName ).Get();
+		check( ArraySetPtr->HasType<AttributeType>() );
+		return static_cast<const TAttributeIndicesArray<AttributeType, ElementIDType>*>( ArraySetPtr )->GetArrayForIndex( AttributeIndex );
 	}
 
-	/** Serializer */
-	friend FArchive& operator<<( FArchive& Ar, TAttributesSet& AttributesSet )
-	{
-		// Serialize the number of attribute types in the container tuple.
-		// If loading, this may be different to the current number defined.
-		const int32 NumAttributeTypes = TTupleArity<AttributeTypes>::Value;
-		int32 SerializedAttributeTypes = NumAttributeTypes;
-		Ar << SerializedAttributeTypes;
-		// Cannot deserialize more attribute types than there are
-		check( NumAttributeTypes >= SerializedAttributeTypes );
-
-		// Serialize the tuple of attribute maps by hand, so we can deserialize correctly when the archive contains fewer tuple elements than the current code
-		// NOTE: This relies on the assumption that VisitTupleElements will always visit elements in ascending order.
-		VisitTupleElements( [ &Ar, SerializedAttributeTypes, TypeIndex = 0, NumElements = 0 ]( auto& AttributesMap ) mutable
-			{
-				if( TypeIndex < SerializedAttributeTypes )
-				{
-					// Serialize attributes map, and keep note of the number of elements present.
-					// This should be the same every iteration.
-					Ar << AttributesMap;
-					check( TypeIndex == 0 || NumElements == AttributesMap.GetNumElements() );
-					NumElements = AttributesMap.GetNumElements();
-				}
-				else
-				{
-					// If we have run out of data to deserialize, initialize this attributes map so that it has the same number of elements
-					// as the other maps.
-					check( Ar.IsLoading() );
-					AttributesMap.Initialize( NumElements );
-				}
-				TypeIndex++;
-			},
-			AttributesSet.Container );
+	template <typename AttributeType>
+	DEPRECATED( 4.20, "Please use GetAttributesRef() or GetAttributesView() instead." )
+	TAttributeIndicesArray<AttributeType, ElementIDType>& GetAttributesSet( const FName AttributeName )
+	{
+		FMeshAttributeArraySetBase* ArraySetPtr = this->Map.FindChecked( AttributeName ).Get();
+		check( ArraySetPtr->HasType<AttributeType>() );
+		return static_cast<TAttributeIndicesArray<AttributeType, ElementIDType>&>( *ArraySetPtr );
+	}
 
-		return Ar;
+	template <typename AttributeType>
+	DEPRECATED( 4.20, "Please use GetAttributesRef() or GetAttributesView() instead." )
+	const TAttributeIndicesArray<AttributeType, ElementIDType>& GetAttributesSet( const FName AttributeName ) const
+	{
+		const FMeshAttributeArraySetBase* ArraySetPtr = this->Map.FindChecked( AttributeName ).Get();
+		check( ArraySetPtr->HasType<AttributeType>() );
+		return static_cast<const TAttributeIndicesArray<AttributeType, ElementIDType>&>( *ArraySetPtr );
 	}
+};
 
-private:
-	/**
-	 * Define type for the entire attribute container.
-	 * We can have attributes of multiple types, each with a name and an arbitrary number of indices,
-	 * whose elements are indexed by an ElementIDType.
-	 *
-	 * This implies the below data structure:
-	 * A TTuple (one per attribute type) of
-	 * TMap keyed on the attribute name,
-	 * yielding a TArray indexed by attribute index,
-	 * yielding a TMeshAttributeArray indexed by an Element ID,
-	 * yielding an item of type AttributeType.
-	 *
-	 * This looks complicated, but actually makes attribute lookup easy when we are interested in a particular attribute for many element IDs.
-	 * By caching the TMeshAttributeArray arrived at by the attribute name and index, we have O(1) access to that attribute for all elements.
-	 */
-	using ContainerType = typename TMakeAttributesSet<ElementIDType, AttributeTypes>::Type;
-	ContainerType Container;
+
+/**
+ * We need a mechanism by which we can iterate all items in the attribute map and perform an arbitrary operation on each.
+ * We require polymorphic behavior, as attribute arrays are templated on their attribute type, and derived from a generic base class.
+ * However, we cannot have a virtual templated method, so we use a different approach.
+ *
+ * Effectively, we wish to cast the attribute array depending on the type member of the base class as we iterate through the map.
+ * This might look something like this:
+ *
+ *    template <typename FuncType>
+ *    void ForEach(FuncType Func)
+ *    {
+ *        for (const auto& MapEntry : Map)
+ *        {
+ *            const uint32 Type = MapEntry.Value->GetType();
+ *            switch (Type)
+ *            {
+ *                case 0: Func(static_cast<TMeshAttributeArraySet<FVector>*>(MapEntry.Value.Get()); break;
+ *                case 1: Func(static_cast<TMeshAttributeArraySet<FVector4>*>(MapEntry.Value.Get()); break;
+ *                case 2: Func(static_cast<TMeshAttributeArraySet<FVector2D>*>(MapEntry.Value.Get()); break;
+ *                case 3: Func(static_cast<TMeshAttributeArraySet<float>*>(MapEntry.Value.Get()); break;
+ *                      ....
+ *            }
+ *        }
+ *    }
+ *
+ * (The hope is that the compiler would optimize the switch into a jump table so we get O(1) dispatch even as the number of attribute types
+ * increases.)
+ *
+ * The approach taken here is to generate a jump table at compile time, one entry per possible attribute type.
+ * The function Dispatch(...) is the actual function which gets called.
+ * MakeJumpTable() is the constexpr function which creates a static jump table at compile time.
+ */
+namespace ForEachImpl
+{
+	// Declare type of jump table used to dispatch functions
+	template <typename ElementIDType, typename ForEachFunc>
+	using JumpTableType = TJumpTable<void( FName, ForEachFunc, FMeshAttributeArraySetBase* ), TTupleArity<AttributeTypes>::Value>;
+
+	// Define dispatch function
+	template <typename ElementIDType, typename ForEachFunc, uint32 I>
+	static void Dispatch( FName Name, ForEachFunc Fn, FMeshAttributeArraySetBase* Attributes )
+	{
+		using AttributeType = typename TTupleElement<I, AttributeTypes>::Type;
+		Fn( Name, TMeshAttributesRef<ElementIDType, AttributeType>( static_cast<TMeshAttributeArraySet<AttributeType>*>( Attributes ) ) );
+	}
+
+	// Build ForEach jump table at compile time, a separate instantiation of Dispatch for each attribute type
+	template <typename ElementIDType, typename ForEachFunc, uint32... Is>
+	static constexpr JumpTableType<ElementIDType, ForEachFunc> MakeJumpTable( TIntegerSequence< uint32, Is...> )
+	{
+		return JumpTableType<ElementIDType, ForEachFunc>( Dispatch<ElementIDType, ForEachFunc, Is>... );
+	}
+}
+
+template <typename ElementIDType>
+template <typename ForEachFunc>
+void TAttributesSet<ElementIDType>::ForEach( ForEachFunc Func )
+{
+	// Construct compile-time jump table for dispatching ForEachImpl::Dispatch() by the attribute type at runtime
+	static constexpr ForEachImpl::JumpTableType<ElementIDType, ForEachFunc>
+		JumpTable = ForEachImpl::MakeJumpTable<ElementIDType, ForEachFunc>( TMakeIntegerSequence<uint32, TTupleArity<AttributeTypes>::Value>() );
+
+	for( auto& MapEntry : this->Map )
+	{
+		const uint32 Type = MapEntry.Value->GetType();
+		JumpTable.Fns[ Type ]( MapEntry.Key, Func, MapEntry.Value.Get() );
+	}
+}
+
+namespace ForEachConstImpl
+{
+	// Declare type of jump table used to dispatch functions
+	template <typename ElementIDType, typename ForEachFunc>
+	using JumpTableType = TJumpTable<void( FName, ForEachFunc, const FMeshAttributeArraySetBase* ), TTupleArity<AttributeTypes>::Value>;
+
+	// Define dispatch function
+	template <typename ElementIDType, typename ForEachFunc, uint32 I>
+	static void Dispatch( FName Name, ForEachFunc Fn, const FMeshAttributeArraySetBase* Attributes )
+	{
+		using AttributeType = typename TTupleElement<I, AttributeTypes>::Type;
+		Fn( Name, TMeshAttributesConstRef<ElementIDType, AttributeType>( static_cast<const TMeshAttributeArraySet<AttributeType>*>( Attributes ) ) );
+	}
+
+	// Build ForEach jump table at compile time, a separate instantiation of Dispatch for each attribute type
+	template <typename ElementIDType, typename ForEachFunc, uint32... Is>
+	static constexpr JumpTableType<ElementIDType, ForEachFunc> MakeJumpTable( TIntegerSequence< uint32, Is...> )
+	{
+		return JumpTableType<ElementIDType, ForEachFunc>( Dispatch<ElementIDType, ForEachFunc, Is>... );
+	}
+}
+
+template <typename ElementIDType>
+template <typename ForEachFunc>
+void TAttributesSet<ElementIDType>::ForEach( ForEachFunc Func ) const
+{
+	// Construct compile-time jump table for dispatching ForEachImpl::Dispatch() by the attribute type at runtime
+	static constexpr ForEachConstImpl::JumpTableType<ElementIDType, ForEachFunc>
+		JumpTable = ForEachConstImpl::MakeJumpTable<ElementIDType, ForEachFunc>( TMakeIntegerSequence<uint32, TTupleArity<AttributeTypes>::Value>() );
+
+	for( const auto& MapEntry : this->Map )
+	{
+		const uint32 Type = MapEntry.Value->GetType();
+		JumpTable.Fns[ Type ]( MapEntry.Key, Func, MapEntry.Value.Get() );
+	}
+}
+
+/**
+ * This is a similar approach to ForEach, above.
+ * Given a type index, at runtime, we wish to create an attribute array of the corresponding type; essentially a factory.
+ *
+ * We generate a jump table at compile time, containing generated functions to register attributes of each type.
+ */
+namespace CreateTypeImpl
+{
+	// Declare type of jump table used to dispatch functions
+	using JumpTableType = TJumpTable<TUniquePtr<FMeshAttributeArraySetBase>(), TTupleArity<AttributeTypes>::Value>;
+
+	// Define dispatch function
+	template <uint32 I>
+	static TUniquePtr<FMeshAttributeArraySetBase> Dispatch()
+	{
+		using AttributeType = typename TTupleElement<I, AttributeTypes>::Type;
+		return MakeUnique<TMeshAttributeArraySet<AttributeType>>();
+	}
+
+	// Build RegisterAttributeOfType jump table at compile time, a separate instantiation of Dispatch for each attribute type
+	template <uint32... Is>
+	static constexpr JumpTableType MakeJumpTable( TIntegerSequence< uint32, Is...> )
+	{
+		return JumpTableType( Dispatch<Is>... );
+	}
+}
+
+inline void FAttributesSetEntry::CreateArrayOfType( const uint32 Type )
+{
+	static constexpr CreateTypeImpl::JumpTableType JumpTable = CreateTypeImpl::MakeJumpTable( TMakeIntegerSequence<uint32, TTupleArity<AttributeTypes>::Value>() );
+	Ptr = JumpTable.Fns[ Type ]();
+}
+
+
+/**
+ * Helper struct which determines whether ViewType and the I'th type from AttributeTypes are mutually constructible from each other.
+ */
+template <typename ViewType, uint32 I>
+struct TIsViewable
+{
+	enum { Value = TIsConstructible<ViewType, typename TTupleElement<I, AttributeTypes>::Type>::Value &&
+				   TIsConstructible<typename TTupleElement<I, AttributeTypes>::Type, ViewType>::Value };
 };
+
+/**
+ * Implementation for TMeshAttributesConstViewBase::Get(ElementIndex).
+ *
+ * This is implemented similarly to the above. A jump table is built, so the correct implementation is dispatched according to the array type.
+ * This cannot be a regular virtual function because the return type depends on the array type.
+ */
+namespace AttributesViewGetImpl
+{
+	// Declare type of jump table used to dispatch functions
+	template <typename ViewType>
+	using JumpTableType = TJumpTable<ViewType( const FMeshAttributeArraySetBase*, int32 ), TTupleArity<AttributeTypes>::Value>;
+
+	// Define dispatch functions
+	template <typename ViewType, uint32 I, typename TEnableIf<TIsViewable<ViewType, I>::Value, int>::Type = 0>
+	static ViewType Dispatch( const FMeshAttributeArraySetBase* Array, const int32 Index )
+	{
+		// Implementation when the attribute type is convertible to the view type
+		return ViewType( static_cast<const TMeshAttributeArraySet<typename TTupleElement<I, AttributeTypes>::Type>*>( Array )->GetArrayForIndex( 0 )[ Index ] );
+	}
+
+	template <typename ViewType, uint32 I, typename TEnableIf<!TIsViewable<ViewType, I>::Value, int>::Type = 0>
+	static ViewType Dispatch( const FMeshAttributeArraySetBase* Array, const int32 Index )
+	{
+		// Implementation when the attribute type is not convertible to the view type
+		check( false );
+		return ViewType();
+	}
+
+	// Build jump table at compile time, a separate instantiation of Dispatch for each view type
+	template <typename ViewType, uint32... Is>
+	static constexpr JumpTableType<ViewType> MakeJumpTable( TIntegerSequence< uint32, Is...> )
+	{
+		return JumpTableType<ViewType>( Dispatch<ViewType, Is>... );
+	}
+}
+
+template <typename ViewType>
+FORCEINLINE ViewType TMeshAttributesViewBase<ViewType>::GetByIndex( const int32 ElementIndex ) const
+{
+	static constexpr AttributesViewGetImpl::JumpTableType<ViewType>
+		JumpTable = AttributesViewGetImpl::MakeJumpTable<ViewType>( TMakeIntegerSequence<uint32, TTupleArity<AttributeTypes>::Value>() );
+
+	return JumpTable.Fns[ ArrayPtr->GetType() ]( ArrayPtr, ElementIndex );
+}
+
+
+/**
+ * Implementation for TMeshAttributesConstViewBase::Get(ElementIndex, AttributeIndex).
+ */
+namespace AttributesViewGetWithIndexImpl
+{
+	// Declare type of jump table used to dispatch functions
+	template <typename ViewType>
+	using JumpTableType = TJumpTable<ViewType( const FMeshAttributeArraySetBase*, int32, int32 ), TTupleArity<AttributeTypes>::Value>;
+
+	// Define dispatch functions
+	template <typename ViewType, uint32 I, typename TEnableIf<TIsViewable<ViewType, I>::Value, int>::Type = 0>
+	static ViewType Dispatch( const FMeshAttributeArraySetBase* Array, const int32 ElementIndex, const int32 AttributeIndex )
+	{
+		return ViewType( static_cast<const TMeshAttributeArraySet<typename TTupleElement<I, AttributeTypes>::Type>*>( Array )->GetArrayForIndex( AttributeIndex )[ ArrayIndex ] );
+	}
+
+	template <typename ViewType, uint32 I, typename TEnableIf<!TIsViewable<ViewType, I>::Value, int>::Type = 0>
+	static ViewType Dispatch( const FMeshAttributeArraySetBase* Array, const int32 ElementIndex, const int32 AttributeIndex )
+	{
+		check( false );
+		return ViewType();
+	}
+
+	// Build jump table at compile time, a separate instantiation of Dispatch for each attribute type
+	template <typename ViewType, uint32... Is>
+	static constexpr JumpTableType<ViewType> MakeJumpTable( TIntegerSequence< uint32, Is...> )
+	{
+		return JumpTableType<ViewType>( Dispatch<ViewType, Is>... );
+	}
+}
+
+template <typename ViewType>
+FORCEINLINE ViewType TMeshAttributesViewBase<ViewType>::GetByIndex( const int32 ElementIndex, const int32 AttributeIndex ) const
+{
+	static constexpr AttributesViewGetWithIndexImpl::JumpTableType<ViewType>
+		JumpTable = AttributesViewGetWithIndexImpl::MakeJumpTable<ViewType>( TMakeIntegerSequence<uint32, TTupleArity<AttributeTypes>::Value>() );
+
+	return JumpTable.Fns[ ArrayPtr->GetType() ]( ArrayPtr, ElementIndex, AttributeIndex );
+}
+
+
+/**
+ * Implementation for TMeshAttributesViewBase::Set(ElementIndex, Value).
+ */
+namespace AttributesViewSetImpl
+{
+	// Declare type of jump table used to dispatch functions
+	template <typename ViewType>
+	using JumpTableType = TJumpTable<void( FMeshAttributeArraySetBase*, int32, const ViewType& ), TTupleArity<AttributeTypes>::Value>;
+
+	// Define dispatch functions
+	template <typename ViewType, uint32 I, typename TEnableIf<TIsViewable<ViewType, I>::Value, int>::Type = 0>
+	static void Dispatch( FMeshAttributeArraySetBase* Array, const int32 Index, const ViewType& Value )
+	{
+		// Implementation when the attribute type is convertible to the view type
+		using AttributeType = typename TTupleElement<I, AttributeTypes>::Type;
+		static_cast<TMeshAttributeArraySet<AttributeType>*>( Array )->GetArrayForIndex( 0 )[ Index ] = AttributeType( Value );
+	}
+
+	template <typename ViewType, uint32 I, typename TEnableIf<!TIsViewable<ViewType, I>::Value, int>::Type = 0>
+	static void Dispatch( FMeshAttributeArraySetBase* Array, const int32 Index, const ViewType& Value )
+	{
+		// Implementation when the attribute type is not convertible to the view type
+		check( false );
+	}
+
+	// Build jump table at compile time, a separate instantiation of Dispatch for each view type
+	template <typename ViewType, uint32... Is>
+	static constexpr JumpTableType<ViewType> MakeJumpTable( TIntegerSequence< uint32, Is...> )
+	{
+		return JumpTableType<ViewType>( Dispatch<ViewType, Is>... );
+	}
+}
+
+template <typename ViewType>
+FORCEINLINE void TMeshAttributesViewBase<ViewType>::SetByIndex( const int32 ElementIndex, const ViewType& Value ) const
+{
+	static constexpr AttributesViewSetImpl::JumpTableType<ViewType>
+		JumpTable = AttributesViewSetImpl::MakeJumpTable<ViewType>( TMakeIntegerSequence<uint32, TTupleArity<AttributeTypes>::Value>() );
+
+	JumpTable.Fns[ this->ArrayPtr->GetType() ]( const_cast<FMeshAttributeArraySetBase*>( this->ArrayPtr ), ElementIndex, Value );
+}
+
+
+/**
+ * Implementation for TMeshAttributesViewBase::Set(ElementIndex, AttributeIndex, Value).
+ */
+namespace AttributesViewSetWithIndexImpl
+{
+	// Declare type of jump table used to dispatch functions
+	template <typename ViewType>
+	using JumpTableType = TJumpTable<void( FMeshAttributeArraySetBase*, int32, int32, const ViewType& ), TTupleArity<AttributeTypes>::Value>;
+
+	// Define dispatch functions
+	template <typename ViewType, uint32 I, typename TEnableIf<TIsViewable<ViewType, I>::Value, int>::Type = 0>
+	static void Dispatch( FMeshAttributeArraySetBase* Array, const int32 ElementIndex, const int32 AttributeIndex, const ViewType& Value )
+	{
+		// Implementation when the attribute type is convertible to the view type
+		using AttributeType = typename TTupleElement<I, AttributeTypes>::Type;
+		static_cast<TMeshAttributeArraySet<AttributeType>*>( Array )->GetArrayForIndex( AttributeIndex )[ ElementIndex ] = AttributeType( Value );
+	}
+
+	template <typename ViewType, uint32 I, typename TEnableIf<!TIsViewable<ViewType, I>::Value, int>::Type = 0>
+	static void Dispatch( FMeshAttributeArraySetBase* Array, const int32 ElementIndex, const int32 AttributeIndex, const ViewType& Value )
+	{
+		// Implementation when the attribute type is not convertible to the view type
+		check( false );
+	}
+
+	// Build jump table at compile time, a separate instantiation of Dispatch for each view type
+	template <typename ViewType, uint32... Is>
+	static constexpr JumpTableType<ViewType> MakeJumpTable( TIntegerSequence< uint32, Is...> )
+	{
+		return JumpTableType<ViewType>( Dispatch<ViewType, Is>... );
+	}
+}
+
+template <typename ViewType>
+FORCEINLINE void TMeshAttributesViewBase<ViewType>::SetByIndex( const int32 ElementIndex, const int32 AttributeIndex, const ViewType& Value ) const
+{
+	static constexpr AttributesViewSetWithIndexImpl::JumpTableType<ViewType>
+		JumpTable = AttributesViewSetWithIndexImpl::MakeJumpTable<ViewType>( TMakeIntegerSequence<uint32, TTupleArity<AttributeTypes>::Value>() );
+
+	JumpTable.Fns[ this->ArrayPtr->GetType() ]( const_cast<FMeshAttributeArraySetBase*>( this->ArrayPtr ), ElementIndex, AttributeIndex, Value );
+}
+
+
+/**
+ * Implementation for TMeshAttributesViewBase::GetDefaultValue().
+ */
+namespace AttributesViewGetDefaultImpl
+{
+	// Declare type of jump table used to dispatch functions
+	template <typename ViewType>
+	using JumpTableType = TJumpTable<ViewType( const FMeshAttributeArraySetBase* ), TTupleArity<AttributeTypes>::Value>;
+
+	// Define dispatch functions
+	template <typename ViewType, uint32 I, typename TEnableIf<TIsViewable<ViewType, I>::Value, int>::Type = 0>
+	static ViewType Dispatch( const FMeshAttributeArraySetBase* Array )
+	{
+		// Implementation when the attribute type is convertible to the view type
+		return ViewType( static_cast<const TMeshAttributeArraySet<typename TTupleElement<I, AttributeTypes>::Type>*>( Array )->GetDefaultValue() );
+	}
+
+	template <typename ViewType, uint32 I, typename TEnableIf<!TIsViewable<ViewType, I>::Value, int>::Type = 0>
+	static ViewType Dispatch( const FMeshAttributeArraySetBase* Array )
+	{
+		// Implementation when the attribute type is not convertible to the view type
+		check( false );
+		return ViewType();
+	}
+
+	// Build jump table at compile time, a separate instantiation of Dispatch for each view type
+	template <typename ViewType, uint32... Is>
+	static constexpr JumpTableType<ViewType> MakeJumpTable( TIntegerSequence< uint32, Is...> )
+	{
+		return JumpTableType<ViewType>( Dispatch<ViewType, Is>... );
+	}
+}
+
+template <typename ViewType>
+FORCEINLINE ViewType TMeshAttributesViewBase<ViewType>::GetDefaultValue() const
+{
+	static constexpr AttributesViewGetDefaultImpl::JumpTableType<ViewType>
+		JumpTable = AttributesViewGetDefaultImpl::MakeJumpTable<ViewType>( TMakeIntegerSequence<uint32, TTupleArity<AttributeTypes>::Value>() );
+
+	return JumpTable.Fns[ ArrayPtr->GetType() ]( ArrayPtr );
+}
diff --git a/Engine/Source/Runtime/MeshDescription/Public/MeshDescription.h b/Engine/Source/Runtime/MeshDescription/Public/MeshDescription.h
--- a/Engine/Source/Runtime/MeshDescription/Public/MeshDescription.h
+++ b/Engine/Source/Runtime/MeshDescription/Public/MeshDescription.h
@@ -8,6 +8,7 @@
 #include "MeshElementArray.h"
 #include "MeshAttributeArray.h"
 #include "UObject/ObjectMacros.h"
+#include "UObject/EditorObjectVersion.h"
 #include "UObject/ReleaseObjectVersion.h"
 #include "MeshDescription.generated.h"
 
@@ -226,6 +227,43 @@ using FEdgeArray = TMeshElementArray<FMeshEdge, FEdgeID>;
 using FPolygonArray = TMeshElementArray<FMeshPolygon, FPolygonID>;
 using FPolygonGroupArray = TMeshElementArray<FMeshPolygonGroup, FPolygonGroupID>;
 
+/** Define aliases for element attributes */
+template <typename AttributeType> using TVertexAttributeIndicesArray = TAttributeIndicesArray<AttributeType, FVertexID>;
+template <typename AttributeType> using TVertexInstanceAttributeIndicesArray = TAttributeIndicesArray<AttributeType, FVertexInstanceID>;
+template <typename AttributeType> using TEdgeAttributeIndicesArray = TAttributeIndicesArray<AttributeType, FEdgeID>;
+template <typename AttributeType> using TPolygonAttributeIndicesArray = TAttributeIndicesArray<AttributeType, FPolygonID>;
+template <typename AttributeType> using TPolygonGroupAttributeIndicesArray = TAttributeIndicesArray<AttributeType, FPolygonGroupID>;
+
+template <typename AttributeType> using TVertexAttributeArray = TMeshAttributeArray<AttributeType, FVertexID>;
+template <typename AttributeType> using TVertexInstanceAttributeArray = TMeshAttributeArray<AttributeType, FVertexInstanceID>;
+template <typename AttributeType> using TEdgeAttributeArray = TMeshAttributeArray<AttributeType, FEdgeID>;
+template <typename AttributeType> using TPolygonAttributeArray = TMeshAttributeArray<AttributeType, FPolygonID>;
+template <typename AttributeType> using TPolygonGroupAttributeArray = TMeshAttributeArray<AttributeType, FPolygonGroupID>;
+
+template <typename AttributeType> using TVertexAttributesRef = TMeshAttributesRef<FVertexID, AttributeType>;
+template <typename AttributeType> using TVertexInstanceAttributesRef = TMeshAttributesRef<FVertexInstanceID, AttributeType>;
+template <typename AttributeType> using TEdgeAttributesRef = TMeshAttributesRef<FEdgeID, AttributeType>;
+template <typename AttributeType> using TPolygonAttributesRef = TMeshAttributesRef<FPolygonID, AttributeType>;
+template <typename AttributeType> using TPolygonGroupAttributesRef = TMeshAttributesRef<FPolygonGroupID, AttributeType>;
+
+template <typename AttributeType> using TVertexAttributesConstRef = TMeshAttributesConstRef<FVertexID, AttributeType>;
+template <typename AttributeType> using TVertexInstanceAttributesConstRef = TMeshAttributesConstRef<FVertexInstanceID, AttributeType>;
+template <typename AttributeType> using TEdgeAttributesConstRef = TMeshAttributesConstRef<FEdgeID, AttributeType>;
+template <typename AttributeType> using TPolygonAttributesConstRef = TMeshAttributesConstRef<FPolygonID, AttributeType>;
+template <typename AttributeType> using TPolygonGroupAttributesConstRef = TMeshAttributesConstRef<FPolygonGroupID, AttributeType>;
+
+template <typename AttributeType> using TVertexAttributesView = TMeshAttributesView<FVertexID, AttributeType>;
+template <typename AttributeType> using TVertexInstanceAttributesView = TMeshAttributesView<FVertexInstanceID, AttributeType>;
+template <typename AttributeType> using TEdgeAttributesView = TMeshAttributesView<FEdgeID, AttributeType>;
+template <typename AttributeType> using TPolygonAttributesView = TMeshAttributesView<FPolygonID, AttributeType>;
+template <typename AttributeType> using TPolygonGroupAttributesView = TMeshAttributesView<FPolygonGroupID, AttributeType>;
+
+template <typename AttributeType> using TVertexAttributesConstView = TMeshAttributesConstView<FVertexID, AttributeType>;
+template <typename AttributeType> using TVertexInstanceAttributesConstView = TMeshAttributesConstView<FVertexInstanceID, AttributeType>;
+template <typename AttributeType> using TEdgeAttributesConstView = TMeshAttributesConstView<FEdgeID, AttributeType>;
+template <typename AttributeType> using TPolygonAttributesConstView = TMeshAttributesConstView<FPolygonID, AttributeType>;
+template <typename AttributeType> using TPolygonGroupAttributesConstView = TMeshAttributesConstView<FPolygonGroupID, AttributeType>;
+
 UENUM()
 enum class EComputeNTBsOptions : uint32
 {
@@ -833,26 +871,26 @@ private:
 	bool ComputePolygonTangentsAndNormals(
 		  const FPolygonID PolygonID
 		, float ComparisonThreshold
-		, const TVertexAttributeArray<FVector>& VertexPositions
-		, const TVertexInstanceAttributeArray<FVector2D>& VertexUVs
-		, TPolygonAttributeArray<FVector>& PolygonNormals
-		, TPolygonAttributeArray<FVector>& PolygonTangents
-		, TPolygonAttributeArray<FVector>& PolygonBinormals
-		, TPolygonAttributeArray<FVector>& PolygonCenters
+		, const TVertexAttributesRef<FVector> VertexPositions
+		, const TVertexInstanceAttributesRef<FVector2D> VertexUVs
+		, TPolygonAttributesRef<FVector> PolygonTangents
+		, TPolygonAttributesRef<FVector> PolygonBinormals
+		, TPolygonAttributesRef<FVector> PolygonCenters
+		, TPolygonAttributesRef<FVector> PolygonNormals
 	);
 
 	void GetVertexConnectedPolygonsInSameSoftEdgedGroup(const FVertexID VertexID, const FPolygonID PolygonID, TArray<FPolygonID>& OutPolygonIDs) const;
 	void GetPolygonsInSameSoftEdgedGroupAsPolygon(const FPolygonID PolygonID, const TArray<FPolygonID>& CandidatePolygonIDs, const TArray<FEdgeID>& SoftEdgeIDs, TArray<FPolygonID>& OutPolygonIDs) const;
 	void GetConnectedSoftEdges(const FVertexID VertexID, TArray<FEdgeID>& OutConnectedSoftEdges) const;
 	void ComputeTangentsAndNormals(
-		  const FVertexInstanceID& VertexInstanceID
+		  const FVertexInstanceID VertexInstanceID
 		, EComputeNTBsOptions ComputeNTBsOptions
-		, const TPolygonAttributeArray<FVector>& PolygonNormals
-		, const TPolygonAttributeArray<FVector>& PolygonTangents
-		, const TPolygonAttributeArray<FVector>& PolygonBinormals
-		, TVertexInstanceAttributeArray<FVector>& VertexNormals
-		, TVertexInstanceAttributeArray<FVector>& VertexTangents
-		, TVertexInstanceAttributeArray<float>& VertexBinormalSigns
+		, const TPolygonAttributesRef<FVector> PolygonNormals
+		, const TPolygonAttributesRef<FVector> PolygonTangents
+		, const TPolygonAttributesRef<FVector> PolygonBinormals
+		, TVertexInstanceAttributesRef<FVector> VertexNormals
+		, TVertexInstanceAttributesRef<FVector> VertexTangents
+		, TVertexInstanceAttributesRef<float> VertexBinormalSigns
 	);
 private:
 
diff --git a/Engine/Source/Runtime/MeshDescription/Public/MeshElementArray.h b/Engine/Source/Runtime/MeshDescription/Public/MeshElementArray.h
--- a/Engine/Source/Runtime/MeshDescription/Public/MeshElementArray.h
+++ b/Engine/Source/Runtime/MeshDescription/Public/MeshElementArray.h
@@ -65,6 +65,13 @@ public:
 		return Ar;
 	}
 
+	/** Compacts elements and returns a remapping table */
+	void Compact( TSparseArray<int32>& OutIndexRemap );
+
+	/** Remaps elements according to the passed remapping table */
+	void Remap( const TSparseArray<int32>& IndexRemap );
+
+
 protected:
 
 	/** The actual container, represented by a sparse array */
@@ -77,7 +84,7 @@ protected:
  * This derived class imposes this type safety.
  */
 template <typename ElementType, typename ElementIDType>
-class TMeshElementArray : public TMeshElementArrayBase<ElementType>
+class TMeshElementArray final : public TMeshElementArrayBase<ElementType>
 {
 	static_assert( TIsDerivedFrom<ElementIDType, FElementID>::IsDerived, "ElementIDType must be derived from FElementID" );
 
@@ -172,12 +179,6 @@ public:
 		return Ar;
 	}
 
-	/** Compacts elements and returns a remapping table */
-	void Compact( TSparseArray<ElementIDType>& OutIndexRemap );
-
-	/** Remaps elements according to the passed remapping table */
-	void Remap( const TSparseArray<ElementIDType>& IndexRemap );
-
 	/**
 	 * This is a special type of iterator which returns successive IDs of valid elements, rather than
 	 * the elements themselves.
@@ -255,8 +256,8 @@ public:
 };
 
 
-template <typename ElementType, typename ElementIDType>
-void TMeshElementArray<ElementType, ElementIDType>::Compact( TSparseArray<ElementIDType>& OutIndexRemap )
+template <typename ElementType>
+void TMeshElementArrayBase<ElementType>::Compact( TSparseArray<int32>& OutIndexRemap )
 {
 	static TSparseArray<ElementType> NewContainer;
 	NewContainer.Empty( Container.Num() );
@@ -273,15 +274,15 @@ void TMeshElementArray<ElementType, ElementIDType>::Compact( TSparseArray<Elemen
 		NewContainer[ NewElementIndex ] = MoveTemp( *It );
 
 		// Provide an O(1) lookup from old index to new index, used when patching up vertex references afterwards
-		OutIndexRemap.Insert( OldElementIndex, ElementIDType( NewElementIndex ) );
+		OutIndexRemap.Insert( OldElementIndex, NewElementIndex );
 	}
 
 	Container = MoveTemp( NewContainer );
 }
 
 
-template <typename ElementType, typename ElementIDType>
-void TMeshElementArray<ElementType, ElementIDType>::Remap( const TSparseArray<ElementIDType>& IndexRemap )
+template <typename ElementType>
+void TMeshElementArrayBase<ElementType>::Remap( const TSparseArray<int32>& IndexRemap )
 {
 	static TSparseArray<ElementType> NewContainer;
 	NewContainer.Empty( IndexRemap.GetMaxIndex() );
@@ -292,7 +293,7 @@ void TMeshElementArray<ElementType, ElementIDType>::Remap( const TSparseArray<El
 		const int32 OldElementIndex = It.GetIndex();
 
 		check( IndexRemap.IsAllocated( OldElementIndex ) );
-		const int32 NewElementIndex = IndexRemap[ OldElementIndex ].GetValue();
+		const int32 NewElementIndex = IndexRemap[ OldElementIndex ];
 
 		// @todo mesheditor: implement TSparseArray::Insert( ElementType&& ) to save this obscure approach
 		NewContainer.Insert( NewElementIndex, ElementType() );
diff --git a/Engine/Source/Runtime/MeshDescription/Public/MeshElementRemappings.h b/Engine/Source/Runtime/MeshDescription/Public/MeshElementRemappings.h
--- a/Engine/Source/Runtime/MeshDescription/Public/MeshElementRemappings.h
+++ b/Engine/Source/Runtime/MeshDescription/Public/MeshElementRemappings.h
@@ -10,40 +10,40 @@
  */
 struct FElementIDRemappings
 {
-	TSparseArray<FVertexID> NewVertexIndexLookup;
-	TSparseArray<FVertexInstanceID> NewVertexInstanceIndexLookup;
-	TSparseArray<FEdgeID> NewEdgeIndexLookup;
-	TSparseArray<FPolygonID> NewPolygonIndexLookup;
-	TSparseArray<FPolygonGroupID> NewPolygonGroupIndexLookup;
+	TSparseArray<int32> NewVertexIndexLookup;
+	TSparseArray<int32> NewVertexInstanceIndexLookup;
+	TSparseArray<int32> NewEdgeIndexLookup;
+	TSparseArray<int32> NewPolygonIndexLookup;
+	TSparseArray<int32> NewPolygonGroupIndexLookup;
 
 	FVertexID GetRemappedVertexID( FVertexID VertexID ) const
 	{
 		checkSlow( NewVertexIndexLookup.IsAllocated( VertexID.GetValue() ) );
-		return NewVertexIndexLookup[ VertexID.GetValue() ];
+		return FVertexID( NewVertexIndexLookup[ VertexID.GetValue() ] );
 	}
 
 	FVertexInstanceID GetRemappedVertexInstanceID( FVertexInstanceID VertexInstanceID ) const
 	{
 		checkSlow( NewVertexInstanceIndexLookup.IsAllocated( VertexInstanceID.GetValue() ) );
-		return NewVertexInstanceIndexLookup[ VertexInstanceID.GetValue() ];
+		return FVertexInstanceID( NewVertexInstanceIndexLookup[ VertexInstanceID.GetValue() ] );
 	}
 
 	FEdgeID GetRemappedEdgeID( FEdgeID EdgeID ) const
 	{
 		checkSlow( NewEdgeIndexLookup.IsAllocated( EdgeID.GetValue() ) );
-		return NewEdgeIndexLookup[ EdgeID.GetValue() ];
+		return FEdgeID( NewEdgeIndexLookup[ EdgeID.GetValue() ] );
 	}
 
 	FPolygonID GetRemappedPolygonID( FPolygonID PolygonID ) const
 	{
 		checkSlow( NewPolygonIndexLookup.IsAllocated( PolygonID.GetValue() ) );
-		return NewPolygonIndexLookup[ PolygonID.GetValue() ];
+		return FPolygonID( NewPolygonIndexLookup[ PolygonID.GetValue() ] );
 	}
 
 	FPolygonGroupID GetRemappedPolygonGroupID( FPolygonGroupID PolygonGroupID ) const
 	{
-		check( NewPolygonGroupIndexLookup.IsAllocated( PolygonGroupID.GetValue() ) );
-		return NewPolygonGroupIndexLookup[ PolygonGroupID.GetValue() ];
+		checkSlow( NewPolygonGroupIndexLookup.IsAllocated( PolygonGroupID.GetValue() ) );
+		return FPolygonGroupID( NewPolygonGroupIndexLookup[ PolygonGroupID.GetValue() ] );
 	}
 };
 
