Fixed typo which caused build errors.

diff --git a/Engine/Source/Runtime/MeshDescription/Public/MeshAttributeArray.h b/Engine/Source/Runtime/MeshDescription/Public/MeshAttributeArray.h
--- a/Engine/Source/Runtime/MeshDescription/Public/MeshAttributeArray.h
+++ b/Engine/Source/Runtime/MeshDescription/Public/MeshAttributeArray.h
@@ -1245,7 +1245,7 @@ namespace AttributesViewGetWithIndexImpl
 	template <typename ViewType, uint32 I, typename TEnableIf<TIsViewable<ViewType, I>::Value, int>::Type = 0>
 	static ViewType Dispatch( const FMeshAttributeArraySetBase* Array, const int32 ElementIndex, const int32 AttributeIndex )
 	{
-		return ViewType( static_cast<const TMeshAttributeArraySet<typename TTupleElement<I, AttributeTypes>::Type>*>( Array )->GetArrayForIndex( AttributeIndex )[ ArrayIndex ] );
+		return ViewType( static_cast<const TMeshAttributeArraySet<typename TTupleElement<I, AttributeTypes>::Type>*>( Array )->GetArrayForIndex( AttributeIndex )[ ElementIndex ] );
 	}
 
 	template <typename ViewType, uint32 I, typename TEnableIf<!TIsViewable<ViewType, I>::Value, int>::Type = 0>
