Removed over-cautious deprecation warnings. Once GetAttributes() is changed to GetAttributesRef(), element access will still work with array syntax.

diff --git a/Engine/Source/Runtime/MeshDescription/Public/MeshAttributeArray.h b/Engine/Source/Runtime/MeshDescription/Public/MeshAttributeArray.h
--- a/Engine/Source/Runtime/MeshDescription/Public/MeshAttributeArray.h
+++ b/Engine/Source/Runtime/MeshDescription/Public/MeshAttributeArray.h
@@ -197,10 +197,7 @@ class TMeshAttributeArray final : public TMeshAttributeArrayBase<AttributeType>
 public:
 
 	/** Element accessors */
-	DEPRECATED( 4.20, "Please access attributes via GetAttributesRef() or GetAttributesView() instead." )
 	FORCEINLINE const AttributeType& operator[]( const ElementIDType ElementID ) const { return this->Container[ ElementID.GetValue() ]; }
-
-	DEPRECATED( 4.20, "Please access attributes via GetAttributesRef() or GetAttributesView() instead." )
 	FORCEINLINE AttributeType& operator[]( const ElementIDType ElementID ) { return this->Container[ ElementID.GetValue() ]; }
 };
 
