Fixed uninitialized member.

diff --git a/Engine/Source/Runtime/MeshDescription/Public/MeshAttributeArray.h b/Engine/Source/Runtime/MeshDescription/Public/MeshAttributeArray.h
--- a/Engine/Source/Runtime/MeshDescription/Public/MeshAttributeArray.h
+++ b/Engine/Source/Runtime/MeshDescription/Public/MeshAttributeArray.h
@@ -689,6 +689,11 @@ private:
 class FAttributesSetBase
 {
 public:
+	/** Constructor */
+	FAttributesSetBase()
+		: NumElements( 0 )
+	{}
+
 	/**
 	 * Register a new attribute name with the given type (must be a member of the AttributeTypes tuple).
 	 * If the attribute name is already registered, it will do nothing.
