Changed FMeshDescription to a struct from a class.
Added log errors when loading UMeshDescription objects (now deprecated), in preparation to resave any which remain.
Once all serialized UMeshDescriptions are wiped out (they only exist internally), FMeshDescription will become a USTRUCT.

diff --git a/Engine/Plugins/Experimental/AlembicImporter/Source/AlembicLibrary/Public/AbcImporter.h b/Engine/Plugins/Experimental/AlembicImporter/Source/AlembicLibrary/Public/AbcImporter.h
--- a/Engine/Plugins/Experimental/AlembicImporter/Source/AlembicLibrary/Public/AbcImporter.h
+++ b/Engine/Plugins/Experimental/AlembicImporter/Source/AlembicLibrary/Public/AbcImporter.h
@@ -26,8 +26,8 @@ class UGeometryCacheTrackStreamable;
 class UAbcImportSettings;
 class FSkeletalMeshImportData;
 class UAbcAssetImportData;
-class FMeshDescription;
 
+struct FMeshDescription;
 struct FAbcImportData;
 struct FGeometryCacheMeshData;
 struct FAbcMeshSample;
diff --git a/Engine/Source/Developer/MeshBuilder/Private/MeshDescriptionHelper.h b/Engine/Source/Developer/MeshBuilder/Private/MeshDescriptionHelper.h
--- a/Engine/Source/Developer/MeshBuilder/Private/MeshDescriptionHelper.h
+++ b/Engine/Source/Developer/MeshBuilder/Private/MeshDescriptionHelper.h
@@ -9,7 +9,7 @@
 DECLARE_LOG_CATEGORY_EXTERN(LogMeshBuilder, Log, All);
 
 class UObject;
-class FMeshDescription;
+struct FMeshDescription;
 struct FMeshBuildSettings;
 struct FVertexInstanceID;
 struct FMeshReductionSettings;
diff --git a/Engine/Source/Developer/MeshDescriptionOperations/Public/MeshDescriptionOperations.h b/Engine/Source/Developer/MeshDescriptionOperations/Public/MeshDescriptionOperations.h
--- a/Engine/Source/Developer/MeshDescriptionOperations/Public/MeshDescriptionOperations.h
+++ b/Engine/Source/Developer/MeshDescriptionOperations/Public/MeshDescriptionOperations.h
@@ -4,7 +4,7 @@
 
 #include "CoreMinimal.h"
 
-class FMeshDescription;
+struct FMeshDescription;
 struct FRawMesh;
 struct FOverlappingCorners;
 enum class ELightmapUVVersion : int32;
diff --git a/Engine/Source/Developer/MeshReductionInterface/Public/IMeshReductionInterfaces.h b/Engine/Source/Developer/MeshReductionInterface/Public/IMeshReductionInterfaces.h
--- a/Engine/Source/Developer/MeshReductionInterface/Public/IMeshReductionInterfaces.h
+++ b/Engine/Source/Developer/MeshReductionInterface/Public/IMeshReductionInterfaces.h
@@ -9,7 +9,7 @@
 
 struct FOverlappingCorners;
 struct FRawMesh;
-class FMeshDescription;
+struct FMeshDescription;
 
 /**
 * Mesh reduction interface.
diff --git a/Engine/Source/Runtime/Engine/Classes/Engine/StaticMesh.h b/Engine/Source/Runtime/Engine/Classes/Engine/StaticMesh.h
--- a/Engine/Source/Runtime/Engine/Classes/Engine/StaticMesh.h
+++ b/Engine/Source/Runtime/Engine/Classes/Engine/StaticMesh.h
@@ -25,7 +25,7 @@ class FSpeedTreeWind;
 class UAssetUserData;
 class UMaterialInterface;
 class UNavCollisionBase;
-class FMeshDescription;
+struct FMeshDescription;
 struct FStaticMeshLODResources;
 
 /*-----------------------------------------------------------------------------
diff --git a/Engine/Source/Runtime/MeshDescription/Private/MeshDescription.cpp b/Engine/Source/Runtime/MeshDescription/Private/MeshDescription.cpp
--- a/Engine/Source/Runtime/MeshDescription/Private/MeshDescription.cpp
+++ b/Engine/Source/Runtime/MeshDescription/Private/MeshDescription.cpp
@@ -8,6 +8,10 @@ void UDEPRECATED_MeshDescription::Serialize( FArchive& Ar )
 {
 	Super::Serialize( Ar );
 
+	if( !HasAnyFlags( RF_ClassDefaultObject ) )
+	{
+		UE_LOG( LogLoad, Error, TEXT( "UMeshDescription about to be deprecated - please resave %s" ), *GetPathName() );
+	}
 	// Discard the contents
 	FMeshDescription MeshDescription;
 	Ar << MeshDescription;
@@ -18,6 +22,11 @@ FArchive& operator<<( FArchive& Ar, FMeshDescription& MeshDescription )
 {
 	Ar.UsingCustomVersion( FReleaseObjectVersion::GUID );
 
+	if( Ar.IsLoading() && Ar.CustomVer( FReleaseObjectVersion::GUID ) < FReleaseObjectVersion::MeshDescriptionNewSerialization )
+	{
+		UE_LOG( LogLoad, Warning, TEXT( "Deprecated serialization format" ) );
+	}
+
 	Ar << MeshDescription.VertexArray;
 	Ar << MeshDescription.VertexInstanceArray;
 	Ar << MeshDescription.EdgeArray;
diff --git a/Engine/Source/Runtime/MeshDescription/Public/MeshDescription.h b/Engine/Source/Runtime/MeshDescription/Public/MeshDescription.h
--- a/Engine/Source/Runtime/MeshDescription/Public/MeshDescription.h
+++ b/Engine/Source/Runtime/MeshDescription/Public/MeshDescription.h
@@ -277,7 +277,8 @@ ENUM_CLASS_FLAGS(EComputeNTBsOptions);
 #define MESHDESCRIPTION_VER TEXT("42190ADE250046AC958045C7DA930FEC")
 
 
-class MESHDESCRIPTION_API FMeshDescription
+//USTRUCT()
+struct MESHDESCRIPTION_API FMeshDescription
 {
 public:
 
