Fix enterprise build, jacketing function need to be port to FMeshDescription

diff --git a/Engine/Plugins/Experimental/ProxyLODPlugin/Source/ProxyLOD/Private/ProxyLODMeshTypes.cpp b/Engine/Plugins/Experimental/ProxyLODPlugin/Source/ProxyLOD/Private/ProxyLODMeshTypes.cpp
--- a/Engine/Plugins/Experimental/ProxyLODPlugin/Source/ProxyLOD/Private/ProxyLODMeshTypes.cpp
+++ b/Engine/Plugins/Experimental/ProxyLODPlugin/Source/ProxyLOD/Private/ProxyLODMeshTypes.cpp
@@ -151,10 +151,6 @@ FMeshDescriptionArrayAdapter::FMeshDescriptionArrayAdapter(const FMeshDescriptio
 
 FMeshDescriptionArrayAdapter::~FMeshDescriptionArrayAdapter()
 {
-	for (FMeshDescription* MeshDescription : RawMeshArray)
-	{
-		delete MeshDescription;
-	}
 	RawMeshArray.clear();
 
 	RawMeshArrayData.clear();
diff --git a/Engine/Source/Runtime/MeshDescription/Public/MeshElementArray.h b/Engine/Source/Runtime/MeshDescription/Public/MeshElementArray.h
--- a/Engine/Source/Runtime/MeshDescription/Public/MeshElementArray.h
+++ b/Engine/Source/Runtime/MeshDescription/Public/MeshElementArray.h
@@ -259,7 +259,7 @@ public:
 template <typename ElementType>
 void TMeshElementArrayBase<ElementType>::Compact( TSparseArray<int32>& OutIndexRemap )
 {
-	static TSparseArray<ElementType> NewContainer;
+	TSparseArray<ElementType> NewContainer;
 	NewContainer.Empty( Container.Num() );
 
 	OutIndexRemap.Empty( Container.GetMaxIndex() );
@@ -284,7 +284,7 @@ void TMeshElementArrayBase<ElementType>::Compact( TSparseArray<int32>& OutIndexR
 template <typename ElementType>
 void TMeshElementArrayBase<ElementType>::Remap( const TSparseArray<int32>& IndexRemap )
 {
-	static TSparseArray<ElementType> NewContainer;
+	TSparseArray<ElementType> NewContainer;
 	NewContainer.Empty( IndexRemap.GetMaxIndex() );
 
 	// Add valid elements into a new contiguous sparse array.  Note non-const iterator so we can move elements.
