MeshAttributeArray: moved GetAttributeNames from TAttributesSet to FAttributesSetBase

diff --git a/Engine/Source/Runtime/MeshDescription/Public/MeshAttributeArray.h b/Engine/Source/Runtime/MeshDescription/Public/MeshAttributeArray.h
--- a/Engine/Source/Runtime/MeshDescription/Public/MeshAttributeArray.h
+++ b/Engine/Source/Runtime/MeshDescription/Public/MeshAttributeArray.h
@@ -810,6 +810,13 @@ public:
 	/** Applies the given remapping to the attributes set */
 	void Remap( const TSparseArray<int32>& IndexRemap );
 
+	/** Returns an array of all the attribute names registered */
+	template <typename Allocator>
+	void GetAttributeNames( TArray<FName, Allocator>& OutAttributeNames ) const
+	{
+		Map.GetKeys( OutAttributeNames );
+	}
+
 	uint32 GetHash(const FName AttributeName) const
 	{
 		if (const FAttributesSetEntry* ArraySetPtr = Map.Find(AttributeName))
@@ -1027,13 +1034,6 @@ public:
 		}
 	}
 
-	/** Returns an array of all the attribute names registered for this attribute type */
-	template <typename AttributeType, typename Allocator>
-	void GetAttributeNames( TArray<FName, Allocator>& OutAttributeNames ) const
-	{
-		this->Map.GetKeys( OutAttributeNames );
-	}
-
 	template <typename AttributeType>
 	AttributeType GetAttribute( const ElementIDType ElementID, const FName AttributeName, const int32 AttributeIndex = 0 ) const
 	{
