Temporary fix for crash when reimporting, pending a robust solution for temporarily unloading bulk data in the editor.

diff --git a/Engine/Source/Runtime/MeshDescription/Private/MeshDescription.cpp b/Engine/Source/Runtime/MeshDescription/Private/MeshDescription.cpp
--- a/Engine/Source/Runtime/MeshDescription/Private/MeshDescription.cpp
+++ b/Engine/Source/Runtime/MeshDescription/Private/MeshDescription.cpp
@@ -1241,7 +1241,8 @@ void FMeshDescriptionBulkData::LoadMeshDescription( FMeshDescription& MeshDescri
 		// Unlock bulk data when we leave scope
 
 		// Throw away the bulk data allocation as we don't need it now we have its contents as a FMeshDescription
-		BulkData.UnloadBulkData();
+		// @todo: revisit this
+//		BulkData.UnloadBulkData();
 	}
 }
 
