fix serialization error

diff --git a/Engine/Source/Runtime/MeshDescription/Private/MeshDescription.cpp b/Engine/Source/Runtime/MeshDescription/Private/MeshDescription.cpp
--- a/Engine/Source/Runtime/MeshDescription/Private/MeshDescription.cpp
+++ b/Engine/Source/Runtime/MeshDescription/Private/MeshDescription.cpp
@@ -23,6 +23,7 @@ void UDEPRECATED_MeshDescription::Serialize( FArchive& Ar )
 FArchive& operator<<( FArchive& Ar, FMeshDescription& MeshDescription )
 {
 	Ar.UsingCustomVersion( FReleaseObjectVersion::GUID );
+	Ar.UsingCustomVersion( FEditorObjectVersion::GUID );
 
 	if( Ar.IsLoading() && Ar.CustomVer( FReleaseObjectVersion::GUID ) < FReleaseObjectVersion::MeshDescriptionNewSerialization )
 	{
diff --git a/Engine/Source/Runtime/MeshDescription/Public/MeshDescription.h b/Engine/Source/Runtime/MeshDescription/Public/MeshDescription.h
--- a/Engine/Source/Runtime/MeshDescription/Public/MeshDescription.h
+++ b/Engine/Source/Runtime/MeshDescription/Public/MeshDescription.h
@@ -183,7 +183,7 @@ struct FMeshPolygon
 	friend FArchive& operator<<( FArchive& Ar, FMeshPolygon& Polygon )
 	{
 		Ar << Polygon.PerimeterContour;
-		if (Ar.IsLoading() && Ar.CustomVer(FReleaseObjectVersion::GUID) < FEditorObjectVersion::MeshDescriptionRemovedHoles)
+		if (Ar.IsLoading() && Ar.CustomVer( FEditorObjectVersion::GUID ) < FEditorObjectVersion::MeshDescriptionRemovedHoles)
 		{
 			TArray<FMeshPolygonContour> Empty;
 			Ar << Empty;
