Fixed handling of MeshDescription transient attributes during transactions. When transacting, the MeshDescription is also serialized to the archive if it is unpacked. Transient MeshDescription attributes are now not removed from the data when writing to a transacting archive.

diff --git a/Engine/Source/Runtime/Engine/Private/StaticMesh.cpp b/Engine/Source/Runtime/Engine/Private/StaticMesh.cpp
--- a/Engine/Source/Runtime/Engine/Private/StaticMesh.cpp
+++ b/Engine/Source/Runtime/Engine/Private/StaticMesh.cpp
@@ -2739,6 +2739,24 @@ void FStaticMeshSourceModel::SerializeBulkData(FArchive& Ar, UObject* Owner)
 
 			MeshDescriptionBulkData->Serialize(Ar, Owner);
 		}
+
+		// For transactions only, serialize the unpacked mesh description here too.
+		// This is so we can preserve any transient attributes which have been set on it when undoing.
+		if (Ar.IsTransacting())
+		{
+			bool bIsMeshDescriptionValid = MeshDescription.IsValid();
+			Ar << bIsMeshDescriptionValid;
+
+			if (bIsMeshDescriptionValid)
+			{
+				if (Ar.IsLoading())
+				{
+					MeshDescription = MakeUnique<FMeshDescription>();
+				}
+
+				Ar << (*MeshDescription);
+			}
+		}
 	}
 }
 #endif // #if WITH_EDITOR
diff --git a/Engine/Source/Runtime/MeshDescription/Private/MeshAttributeArray.cpp b/Engine/Source/Runtime/MeshDescription/Private/MeshAttributeArray.cpp
--- a/Engine/Source/Runtime/MeshDescription/Private/MeshAttributeArray.cpp
+++ b/Engine/Source/Runtime/MeshDescription/Private/MeshAttributeArray.cpp
@@ -68,7 +68,7 @@ FArchive& operator<<( FArchive& Ar, FAttributesSetBase& AttributesSet )
 
 	// If saving, store transient attribute arrays and remove them temporarily from the map
 	TArray<TTuple<FName, FAttributesSetEntry>> TransientArrays;
-	if( Ar.IsSaving() )
+	if( Ar.IsSaving() && !Ar.IsTransacting() )
 	{
 		for( TMap<FName, FAttributesSetEntry>::TIterator It( AttributesSet.Map ); It; ++It )
 		{
@@ -84,7 +84,7 @@ FArchive& operator<<( FArchive& Ar, FAttributesSetBase& AttributesSet )
 	Ar << AttributesSet.Map;
 
 	// Restore transient attribute arrays if saving
-	if( Ar.IsSaving() )
+	if( Ar.IsSaving() && !Ar.IsTransacting() )
 	{
 		for( auto& TransientArray : TransientArrays )
 		{
