Preserve FMeshDescriptionBulkData versioning information when loaded from package, so it can be retrospectively applied to the FBulkDataReader archive when deserializing FMeshDescription.

diff --git a/Engine/Source/Runtime/MeshDescription/Private/MeshDescription.cpp b/Engine/Source/Runtime/MeshDescription/Private/MeshDescription.cpp
--- a/Engine/Source/Runtime/MeshDescription/Private/MeshDescription.cpp
+++ b/Engine/Source/Runtime/MeshDescription/Private/MeshDescription.cpp
@@ -1226,6 +1226,11 @@ void FMeshDescription::ReverseAllPolygonFacing()
 
 void FMeshDescriptionBulkData::Serialize( FArchive& Ar, UObject* Owner )
 {
+	if (Ar.IsLoading())
+	{
+		CustomVersions = Ar.GetCustomVersions();
+	}
+
 	BulkData.Serialize( Ar, Owner );
 
 	if( Ar.IsLoading() && Ar.CustomVer( FEditorObjectVersion::GUID ) < FEditorObjectVersion::MeshDescriptionBulkDataGuid )
@@ -1264,6 +1269,11 @@ void FMeshDescriptionBulkData::LoadMeshDescription( FMeshDescription& MeshDescri
 		{
 			const bool bIsPersistent = true;
 			FBulkDataReader Ar( BulkData, bIsPersistent );
+
+			// Propagate the custom version information from the package to the bulk data, so that the MeshDescription
+			// is serialized with the same versioning.
+			Ar.SetCustomVersions(CustomVersions);
+
 			Ar << MeshDescription;
 		}
 		// Unlock bulk data when we leave scope
diff --git a/Engine/Source/Runtime/MeshDescription/Public/MeshDescription.h b/Engine/Source/Runtime/MeshDescription/Public/MeshDescription.h
--- a/Engine/Source/Runtime/MeshDescription/Public/MeshDescription.h
+++ b/Engine/Source/Runtime/MeshDescription/Public/MeshDescription.h
@@ -12,6 +12,7 @@
 #include "UObject/EditorObjectVersion.h"
 #include "UObject/ReleaseObjectVersion.h"
 #include "Serialization/BulkData.h"
+#include "Serialization/CustomVersion.h"
 #include "MeshDescription.generated.h"
 
 enum
@@ -922,6 +923,8 @@ private:
 
 	/** GUID associated with the data stored herein. */
 	FGuid Guid;
+
+	FCustomVersionContainer CustomVersions;
 };
 
 
