Merge conflict.
richard.talbotwatkin, please merge this change by hand.
More info at https://robomerge.epicgames.net#RELEASE

/src/ROBOMERGE_RELEASE_Main/Engine/Source/Runtime/MeshDescription/Private/MeshDescription.cpp - merging //UE4/Release-4.22/Engine/Source/Runtime/MeshDescription/Private/MeshDescription.cpp#2

--------------------------------------
Copied from //UE4/Main CL 4879714

Preserve FMeshDescriptionBulkData versioning information when loaded from package, so it can be retrospectively applied to the FBulkDataReader archive when deserializing FMeshDescription.

diff --git a/Engine/Source/Runtime/MeshDescription/Public/MeshDescription.h b/Engine/Source/Runtime/MeshDescription/Public/MeshDescription.h
--- a/Engine/Source/Runtime/MeshDescription/Public/MeshDescription.h
+++ b/Engine/Source/Runtime/MeshDescription/Public/MeshDescription.h
@@ -924,6 +924,7 @@ private:
 	/** GUID associated with the data stored herein. */
 	FGuid Guid;
 
+	/** Take a copy of the bulk data versioning so it can be propagated to the bulk data reader when deserializing MeshDescription */
 	FCustomVersionContainer CustomVersions;
 };
 
