Fix serialization issues with MeshDescription.

FMeshDescriptionBulkData is serialized directly to a package, but its bulk data contains a compressed FMeshDescription which is serialized at a later stage on demand. It is mandatory that the versioning for each of these stay in sync, otherwise MeshDescription will fail to load following a version bump.

This fix ensures that, when saving FMeshDescriptionBulkData, its bulk data is reserialized with the latest versioning. It is not always the case that the mesh description would be otherwise updated prior to saving the FMeshDescriptionBulkData.

This fix also ensures that the internal state of FMeshDescriptionBulkData is saved and restored around an undo/redo. Previously this could get lost when performing an undo in the static mesh editor.

diff --git a/Engine/Source/Runtime/MeshDescription/Private/MeshDescription.cpp b/Engine/Source/Runtime/MeshDescription/Private/MeshDescription.cpp
--- a/Engine/Source/Runtime/MeshDescription/Private/MeshDescription.cpp
+++ b/Engine/Source/Runtime/MeshDescription/Private/MeshDescription.cpp
@@ -1226,9 +1226,33 @@ void FMeshDescription::ReverseAllPolygonFacing()
 
 void FMeshDescriptionBulkData::Serialize( FArchive& Ar, UObject* Owner )
 {
-	if( Ar.IsLoading() )
+	Ar.UsingCustomVersion( FEditorObjectVersion::GUID );
+
+	if( Ar.IsTransacting() )
+	{
+		// If transacting, keep these members alive the other side of an undo, otherwise their values will get lost
+		CustomVersions.Serialize( Ar );
+		Ar << bBulkDataUpdated;
+	}
+	else
 	{
-		CustomVersions = Ar.GetCustomVersions();
+		if( Ar.IsLoading() )
+		{
+			// If loading, take a copy of the package custom version container, so it can be applied when unpacking
+			// MeshDescription from the bulk data.
+			CustomVersions = Ar.GetCustomVersions();
+		}
+		else if( Ar.IsSaving() )
+		{
+			// If the bulk data hasn't been updated since this was loaded, there's a possibility that it has old versioning.
+			// Explicitly load and resave the FMeshDescription so that its version is in sync with the FMeshDescriptionBulkData.
+			if( !bBulkDataUpdated )
+			{
+				FMeshDescription MeshDescription;
+				LoadMeshDescription( MeshDescription );
+				SaveMeshDescription( MeshDescription );
+			}
+		}
 	}
 
 	BulkData.Serialize( Ar, Owner );
@@ -1256,6 +1280,10 @@ void FMeshDescriptionBulkData::SaveMeshDescription( FMeshDescription& MeshDescri
 	}
 
 	FPlatformMisc::CreateGuid( Guid );
+
+	// Mark the MeshDescriptionBulkData as having been updated.
+	// This means we know that its version is up-to-date.
+	bBulkDataUpdated = true;
 }
 
 
diff --git a/Engine/Source/Runtime/MeshDescription/Public/MeshDescription.h b/Engine/Source/Runtime/MeshDescription/Public/MeshDescription.h
--- a/Engine/Source/Runtime/MeshDescription/Public/MeshDescription.h
+++ b/Engine/Source/Runtime/MeshDescription/Public/MeshDescription.h
@@ -893,6 +893,7 @@ struct MESHDESCRIPTION_API FMeshDescriptionBulkData
 {
 public:
 	FMeshDescriptionBulkData()
+		: bBulkDataUpdated( false )
 	{
 		BulkData.SetBulkDataFlags( BULKDATA_SerializeCompressed | BULKDATA_SerializeCompressedBitWindow );
 	}
@@ -926,6 +927,9 @@ private:
 
 	/** Take a copy of the bulk data versioning so it can be propagated to the bulk data reader when deserializing MeshDescription */
 	FCustomVersionContainer CustomVersions;
+
+	/** Whether the bulk data has been written via SaveMeshDescription */
+	bool bBulkDataUpdated;
 };
 
 
