Copying //UE4/Dev-Editor to Dev-Main (//UE4/Dev-Main) for 4.24

diff --git a/Engine/Source/Runtime/Engine/Private/SkeletalMeshRenderData.cpp b/Engine/Source/Runtime/Engine/Private/SkeletalMeshRenderData.cpp
--- a/Engine/Source/Runtime/Engine/Private/SkeletalMeshRenderData.cpp
+++ b/Engine/Source/Runtime/Engine/Private/SkeletalMeshRenderData.cpp
@@ -15,6 +15,8 @@
 #include "Interfaces/ITargetPlatformManagerModule.h"
 #include "Interfaces/ITargetPlatform.h"
 #include "PlatformInfo.h"
+#include "IMeshBuilderModule.h"
+#include "EngineUtils.h"
 
 
 #if ENABLE_COOK_STATS
@@ -28,11 +30,36 @@ namespace SkeletalMeshCookStats
 }
 #endif
 
+//Serialize the LODInfo and append the result to the KeySuffix to build the LODInfo part of the DDC KEY
+//Note: this serializer is only used to build the mesh DDC key, no versioning is required
+static void SerializeLODInfoForDDC(USkeletalMesh* SkeletalMesh, FString& KeySuffix)
+{
+	TArray<FSkeletalMeshLODInfo>& LODInfos = SkeletalMesh->GetLODInfoArray();
+	const bool bIs16BitfloatBufferSupported = GVertexElementTypeSupport.IsSupported(VET_Half2);
+	for (int32 LODIndex = 0; LODIndex < SkeletalMesh->GetLODNum(); ++LODIndex)
+	{
+		check(LODInfos.IsValidIndex(LODIndex));
+		FSkeletalMeshLODInfo& LODInfo = LODInfos[LODIndex];
+		bool bValidLODSettings = false;
+		if (SkeletalMesh->LODSettings != nullptr)
+		{
+			const int32 NumSettings = FMath::Min(SkeletalMesh->LODSettings->GetNumberOfSettings(), SkeletalMesh->GetLODNum());
+			if (LODIndex < NumSettings)
+			{
+				bValidLODSettings = true;
+			}
+		}
+		const FSkeletalMeshLODGroupSettings* SkeletalMeshLODGroupSettings = bValidLODSettings ? &SkeletalMesh->LODSettings->GetSettingsForLODLevel(LODIndex) : nullptr;
+		LODInfo.BuildGUID = LODInfo.ComputeDeriveDataCacheKey(SkeletalMeshLODGroupSettings);
+		KeySuffix += LODInfo.BuildGUID.ToString(EGuidFormats::Digits);
+	}
+}
+
 // If skeletal mesh derived data needs to be rebuilt (new format, serialization
 // differences, etc.) replace the version GUID below with a new one.
 // In case of merge conflicts with DDC versions, you MUST generate a new GUID
 // and set this new GUID as the version.                                       
-#define SKELETALMESH_DERIVEDDATA_VER TEXT("4F54F64FD82B40C7A9F4158AF684B64D")
+#define SKELETALMESH_DERIVEDDATA_VER TEXT("51BDAF351B6941548818C2F8B33397D7")
 
 static const FString& GetSkeletalMeshDerivedDataVersion()
 {
@@ -48,8 +75,33 @@ static FString BuildSkeletalMeshDerivedDataKey(USkeletalMesh* SkelMesh)
 {
 	FString KeySuffix(TEXT(""));
 
-	KeySuffix += SkelMesh->GetImportedModel()->GetIdString();
-	KeySuffix += (SkelMesh->bUseFullPrecisionUVs || !GVertexElementTypeSupport.IsSupported(VET_Half2)) ? "1" : "0";
+	if (SkelMesh->UseLegacyMeshDerivedDataKey )
+	{
+		//Old asset will have the same LOD settings for bUseFullPrecisionUVs. We can use the LOD 0
+		const FSkeletalMeshLODInfo* BaseLODInfo = SkelMesh->GetLODInfo(0);
+		bool bUseFullPrecisionUVs = BaseLODInfo ? BaseLODInfo->BuildSettings.bUseFullPrecisionUVs : false;
+		KeySuffix += SkelMesh->GetImportedModel()->GetIdString();
+		KeySuffix += (bUseFullPrecisionUVs || !GVertexElementTypeSupport.IsSupported(VET_Half2)) ? "1" : "0";
+	}
+	else
+	{
+		FString tmpDebugString;
+		//Synchronize the user data that are part of the key
+		SkelMesh->GetImportedModel()->SyncronizeLODUserSectionsData();
+		tmpDebugString = SkelMesh->GetImportedModel()->GetIdString();
+		KeySuffix += tmpDebugString;
+		tmpDebugString = SkelMesh->GetImportedModel()->GetLODModelIdString();
+		KeySuffix += tmpDebugString;
+		
+		//Add the max gpu bone per section
+		const uint32 MaxGPUSkinBones = FGPUBaseSkinVertexFactory::GetMaxGPUSkinBones();
+		KeySuffix += FString::FromInt((int32)MaxGPUSkinBones);
+
+		tmpDebugString = TEXT("");
+		SerializeLODInfoForDDC(SkelMesh, tmpDebugString);
+		KeySuffix += tmpDebugString;
+	}
+
 	KeySuffix += SkelMesh->bHasVertexColors ? "1" : "0";
 	KeySuffix += SkelMesh->VertexColorGuid.ToString(EGuidFormats::Digits);
 
@@ -99,7 +151,34 @@ void FSkeletalMeshRenderData::Cache(USkeletalMesh* Owner)
 		if (GetDerivedDataCacheRef().GetSynchronous(*DerivedDataKey, DerivedData))
 		{
 			COOK_STAT(Timer.AddHit(DerivedData.Num()));
+			
+
 			FMemoryReader Ar(DerivedData, /*bIsPersistent=*/ true);
+
+			//With skeletal mesh build refactor we serialize the LODModel sections into the DDC
+			//We need to store those so we do not have to rerun the reduction to make them up to date
+			//with the serialize renderdata. This allow to use DDC when changing the reduction settings.
+			//The old workflow has to reduce the LODModel before getting the render data DDC.
+			if (!Owner->UseLegacyMeshDerivedDataKey)
+			{
+				FSkeletalMeshModel* SkelMeshModel = Owner->GetImportedModel();
+				check(SkelMeshModel);
+				//Serialize the LODModel sections since they are dependent on the reduction
+				for (int32 LODIndex = 0; LODIndex < SkelMeshModel->LODModels.Num(); LODIndex++)
+				{
+					FSkeletalMeshLODModel* LODModel = &(SkelMeshModel->LODModels[LODIndex]);
+					int32 SectionNum = 0;
+					Ar << SectionNum;
+					LODModel->Sections.Reset(SectionNum);
+					LODModel->Sections.AddDefaulted(SectionNum);
+					for (int32 SectionIndex = 0; SectionIndex < LODModel->Sections.Num(); ++SectionIndex)
+					{
+						Ar << LODModel->Sections[SectionIndex];
+					}
+					LODModel->SyncronizeUserSectionsDataArray();
+				}
+			}
+
 			Serialize(Ar, Owner);
 			for (int32 LODIndex = 0; LODIndex < LODRenderData.Num(); ++LODIndex)
 			{
@@ -129,14 +208,77 @@ void FSkeletalMeshRenderData::Cache(USkeletalMesh* Owner)
 
 			for (int32 LODIndex = 0; LODIndex < SkelMeshModel->LODModels.Num(); LODIndex++)
 			{
-				FSkeletalMeshLODModel& LODModel = SkelMeshModel->LODModels[LODIndex];
+				FSkeletalMeshLODModel* LODModel = &(SkelMeshModel->LODModels[LODIndex]);
+				FSkeletalMeshLODInfo* LODInfo = Owner->GetLODInfo(LODIndex);
+				check(LODInfo);
+				//Build the source model before the render data, if we are a purely generated LOD we do not need to be build
+				IMeshBuilderModule& MeshBuilderModule = FModuleManager::Get().LoadModuleChecked<IMeshBuilderModule>("MeshBuilder");
+				if (!LODModel->RawSkeletalMeshBulkData.IsEmpty() && LODModel->RawSkeletalMeshBulkData.IsBuildDataAvailable())
+				{
+					MeshBuilderModule.BuildSkeletalMesh(Owner, LODIndex, true);
+					LODModel = &(SkelMeshModel->LODModels[LODIndex]);
+				}
+				else
+				{
+					//We need to synchronize when we are generated mesh or if we have load an old asset that was not re-imported
+					LODModel->SyncronizeUserSectionsDataArray();
+				}
+
 				FSkeletalMeshLODRenderData* LODData = new FSkeletalMeshLODRenderData();
 				LODRenderData.Add(LODData);
-
-				LODData->BuildFromLODModel(&LODModel, VertexBufferBuildFlags);
+				
+				//Get the UVs and tangents precision build settings flag specific for this LOD index
+				{
+					bool bUseFullPrecisionUVs = LODInfo->BuildSettings.bUseFullPrecisionUVs;
+					bool bUseHighPrecisionTangentBasis = LODInfo->BuildSettings.bUseHighPrecisionTangentBasis;
+					bool bBuildAdjacencyBuffer = LODInfo->BuildSettings.bBuildAdjacencyBuffer;
+					if (Owner->LODSettings != nullptr)
+					{
+						const int32 NumSettings = Owner->LODSettings->GetNumberOfSettings();
+						if (LODIndex < NumSettings)
+						{
+							const FSkeletalMeshLODGroupSettings& SkeletalMeshLODGroupSettings = Owner->LODSettings->GetSettingsForLODLevel(LODIndex);
+							bUseFullPrecisionUVs = SkeletalMeshLODGroupSettings.GetBuildSettings().bUseFullPrecisionUVs;
+							bUseHighPrecisionTangentBasis = SkeletalMeshLODGroupSettings.GetBuildSettings().bUseHighPrecisionTangentBasis;
+							bBuildAdjacencyBuffer = SkeletalMeshLODGroupSettings.GetBuildSettings().bBuildAdjacencyBuffer;
+						}
+					}
+					if (bUseFullPrecisionUVs || !GVertexElementTypeSupport.IsSupported(VET_Half2))
+					{
+						VertexBufferBuildFlags |= ESkeletalMeshVertexFlags::UseFullPrecisionUVs;
+					}
+					if (bUseHighPrecisionTangentBasis)
+					{
+						VertexBufferBuildFlags |= ESkeletalMeshVertexFlags::UseHighPrecisionTangentBasis;
+					}
+					if (bBuildAdjacencyBuffer)
+					{
+						VertexBufferBuildFlags |= ESkeletalMeshVertexFlags::BuildAdjacencyIndexBuffer;
+					}
+				}
+				LODData->BuildFromLODModel(LODModel, VertexBufferBuildFlags);
 			}
 
 			FMemoryWriter Ar(DerivedData, /*bIsPersistent=*/ true);
+			
+			//If we load an old asset we want to be sure the serialize ddc will be the same has before the skeletalmesh build refactor
+			//So we do not serialize the LODModel sections.
+			if (!Owner->UseLegacyMeshDerivedDataKey)
+			{
+				//Serialize the LODModel sections since they are dependent on the reduction
+				for (int32 LODIndex = 0; LODIndex < SkelMeshModel->LODModels.Num(); LODIndex++)
+				{
+					FSkeletalMeshLODModel* LODModel = &(SkelMeshModel->LODModels[LODIndex]);
+					int32 SectionNum = LODModel->Sections.Num();
+					Ar << SectionNum;
+					for (int32 SectionIndex = 0; SectionIndex < LODModel->Sections.Num(); ++SectionIndex)
+					{
+						Ar << LODModel->Sections[SectionIndex];
+					}
+				}
+			}
+
+			//Serialize the render data
 			Serialize(Ar, Owner);
 			for (int32 LODIndex = 0; LODIndex < LODRenderData.Num(); ++LODIndex)
 			{
diff --git a/Engine/Source/Runtime/Engine/Public/Rendering/SkeletalMeshLODImporterData.cpp b/Engine/Source/Runtime/Engine/Public/Rendering/SkeletalMeshLODImporterData.cpp
--- a/Engine/Source/Runtime/Engine/Public/Rendering/SkeletalMeshLODImporterData.cpp
+++ b/Engine/Source/Runtime/Engine/Public/Rendering/SkeletalMeshLODImporterData.cpp
@@ -3,7 +3,8 @@
 #if WITH_EDITOR
 
 #include "SkeletalMeshLODImporterData.h"
-#include "Serialization/MemoryWriter.h"
+#include "Serialization/BulkDataWriter.h"
+#include "Serialization/BulkDataReader.h"
 #include "Rendering/SkeletalMeshModel.h"
 #include "Engine/SkeletalMesh.h"
 #include "Factories/FbxSkeletalMeshImportData.h"
@@ -349,14 +350,16 @@ enum
 
 struct FReductionSkeletalMeshData
 {
-	FReductionSkeletalMeshData(FSkeletalMeshLODModel& InBaseLODModel, TMap<FString, TArray<FMorphTargetDelta>>& InBaseLODMorphTargetData)
+	FReductionSkeletalMeshData(FSkeletalMeshLODModel& InBaseLODModel, TMap<FString, TArray<FMorphTargetDelta>>& InBaseLODMorphTargetData, UObject* InOwner)
 		: BaseLODModel(InBaseLODModel)
 		, BaseLODMorphTargetData(InBaseLODMorphTargetData)
+		, Owner(InOwner)
 	{
 	}
 
 	FSkeletalMeshLODModel& BaseLODModel;
 	TMap<FString, TArray<FMorphTargetDelta>>& BaseLODMorphTargetData;
+	UObject* Owner;
 };
 
 FArchive& operator<<(FArchive& Ar, FReductionSkeletalMeshData& ReductionSkeletalMeshData)
@@ -365,7 +368,7 @@ FArchive& operator<<(FArchive& Ar, FReductionSkeletalMeshData& ReductionSkeletal
 	int32 LicenseeVersion = REDUCTION_BASE_SK_DATA_BULKDATA_LIC_VER;
 	Ar << Version;
 	Ar << LicenseeVersion;
-	ReductionSkeletalMeshData.BaseLODModel.Serialize(Ar, nullptr, 0);
+	ReductionSkeletalMeshData.BaseLODModel.Serialize(Ar, ReductionSkeletalMeshData.Owner, INDEX_NONE);
 	Ar << ReductionSkeletalMeshData.BaseLODMorphTargetData;
 	return Ar;
 }
@@ -386,10 +389,8 @@ void FReductionBaseSkeletalMeshBulkData::Serialize(FArchive& Ar, TArray<FReducti
 		for (int32 Index = 0; Index < NewNum; Index++)
 		{
 			FReductionBaseSkeletalMeshBulkData* EmptyData = new FReductionBaseSkeletalMeshBulkData();
-			ReductionBaseSkeletalMeshDatas.Add(EmptyData);
-		}
-		for (int32 Index = 0; Index < NewNum; Index++)
-		{
+			int32 NewEntryIndex = ReductionBaseSkeletalMeshDatas.Add(EmptyData);
+			check(NewEntryIndex == Index);
 			ReductionBaseSkeletalMeshDatas[Index]->Serialize(Ar, Owner);
 		}
 	}
@@ -407,68 +408,75 @@ void FReductionBaseSkeletalMeshBulkData::Serialize(FArchive& Ar, TArray<FReducti
 
 void FReductionBaseSkeletalMeshBulkData::Serialize(FArchive& Ar, UObject* Owner)
 {
-	if (Ar.IsLoading())
+	if (Ar.IsTransacting())
 	{
-		//Save the custom version so we can load FReductionSkeletalMeshData later
-		SerializeLoadingCustomVersionContainer = Ar.GetCustomVersions();
-		bUseSerializeLoadingCustomVersion = false;
-		const FCustomVersionContainer& RegisteredCustomVersionContainer = FCustomVersionContainer::GetRegistered();
-		const FCustomVersionArray& ArchiveCustomVersionArray = SerializeLoadingCustomVersionContainer.GetAllVersions();
-		for (const FCustomVersion& ArchiveVersion : ArchiveCustomVersionArray)
+		// If transacting, keep these members alive the other side of an undo, otherwise their values will get lost
+		SerializeLoadingCustomVersionContainer.Serialize(Ar);
+		Ar << bUseSerializeLoadingCustomVersion;
+	}
+	else
+	{
+		if (Ar.IsLoading())
 		{
-			const FCustomVersion* RegisteredVersion = RegisteredCustomVersionContainer.GetVersion(ArchiveVersion.Key);
-			if (!RegisteredVersion || RegisteredVersion->Version != ArchiveVersion.Version)
-			{
-				bUseSerializeLoadingCustomVersion = true;
-				break;
-			}
+			//Save the custom version so we can load FReductionSkeletalMeshData later
+			SerializeLoadingCustomVersionContainer = Ar.GetCustomVersions();
+			bUseSerializeLoadingCustomVersion = true;
 		}
-	}
 
-	if (Ar.IsSaving() && bUseSerializeLoadingCustomVersion == true)
-	{
-		//We need to update the FReductionSkeletalMeshData serialize version to the latest in case we save the Parent bulkdata
-		FSkeletalMeshLODModel BaseLODModel;
-		TMap<FString, TArray<FMorphTargetDelta>> BaseLODMorphTargetData;
-		LoadReductionData(BaseLODModel, BaseLODMorphTargetData);
-		SaveReductionData(BaseLODModel, BaseLODMorphTargetData);
+		if (Ar.IsSaving() && bUseSerializeLoadingCustomVersion == true)
+		{
+			//We need to update the FReductionSkeletalMeshData serialize version to the latest in case we save the Parent bulkdata
+			FSkeletalMeshLODModel BaseLODModel;
+			TMap<FString, TArray<FMorphTargetDelta>> BaseLODMorphTargetData;
+			LoadReductionData(BaseLODModel, BaseLODMorphTargetData, Owner);
+			SaveReductionData(BaseLODModel, BaseLODMorphTargetData, Owner);
+		}
 	}
 
 	BulkData.Serialize(Ar, Owner);
 }
 
-void FReductionBaseSkeletalMeshBulkData::SaveReductionData(FSkeletalMeshLODModel& BaseLODModel, TMap<FString, TArray<FMorphTargetDelta>>& BaseLODMorphTargetData)
+void FReductionBaseSkeletalMeshBulkData::SaveReductionData(FSkeletalMeshLODModel& BaseLODModel, TMap<FString, TArray<FMorphTargetDelta>>& BaseLODMorphTargetData, UObject* Owner)
 {
 	//Saving the bulk data mean we do not need anymore the SerializeLoadingCustomVersionContainer of the parent bulk data
 	SerializeLoadingCustomVersionContainer.Empty();
 	bUseSerializeLoadingCustomVersion = false;
-	FReductionSkeletalMeshData ReductionSkeletalMeshData(BaseLODModel, BaseLODMorphTargetData);
-	TArray<uint8> TempBytes;
-	FMemoryWriter Ar(TempBytes, /*bIsPersistent=*/ true);
-	Ar << ReductionSkeletalMeshData;
-	BulkData.Lock(LOCK_READ_WRITE);
-	uint8* Dest = (uint8*)BulkData.Realloc(TempBytes.Num());
-	FMemory::Memcpy(Dest, TempBytes.GetData(), TempBytes.Num());
-	BulkData.Unlock();
+	FReductionSkeletalMeshData ReductionSkeletalMeshData(BaseLODModel, BaseLODMorphTargetData, Owner);
+
+	//Clear the bulk data before writing it
+	BulkData.RemoveBulkData();
+
+	// Get a lock on the bulk data
+	{
+		const bool bIsPersistent = true;
+		FBulkDataWriter Ar(BulkData, bIsPersistent);
+		Ar << ReductionSkeletalMeshData;
+
+		// Preserve CustomVersions at save time so we can reuse the same ones when reloading direct from memory
+		SerializeLoadingCustomVersionContainer = Ar.GetCustomVersions();
+	}
+	// Unlock the bulk data
 }
 
-void FReductionBaseSkeletalMeshBulkData::LoadReductionData(FSkeletalMeshLODModel& BaseLODModel, TMap<FString, TArray<FMorphTargetDelta>>& BaseLODMorphTargetData)
+void FReductionBaseSkeletalMeshBulkData::LoadReductionData(FSkeletalMeshLODModel& BaseLODModel, TMap<FString, TArray<FMorphTargetDelta>>& BaseLODMorphTargetData, UObject* Owner)
 {
 	BaseLODMorphTargetData.Empty();
 	if (BulkData.GetElementCount() > 0)
 	{
-		FReductionSkeletalMeshData ReductionSkeletalMeshData(BaseLODModel, BaseLODMorphTargetData);
-		FBufferReader Ar(
-			BulkData.Lock(LOCK_READ_ONLY), BulkData.GetElementCount(),
-			/*bInFreeOnClose=*/ false, /*bIsPersistent=*/ true
-		);
-		//If we load the bulkdata content and the bulk was using customVersion, we must set the same until we save the reduction data
-		if (bUseSerializeLoadingCustomVersion)
+		FReductionSkeletalMeshData ReductionSkeletalMeshData(BaseLODModel, BaseLODMorphTargetData, Owner);
+
+		// Get a lock on the bulk data
 		{
+			const bool bIsPersistent = true;
+			FBulkDataReader Ar(BulkData, bIsPersistent);
+			
+			// Propagate the custom version information from the package to the bulk data, so that the MeshDescription
+			// is serialized with the same versioning.
 			Ar.SetCustomVersions(SerializeLoadingCustomVersionContainer);
+
+			Ar << ReductionSkeletalMeshData;
 		}
-		Ar << ReductionSkeletalMeshData;
-		BulkData.Unlock();
+		// Unlock the bulk data
 	}
 }
 
@@ -479,6 +487,8 @@ FRawSkeletalMeshBulkData
 FRawSkeletalMeshBulkData::FRawSkeletalMeshBulkData()
 	: bGuidIsHash(false)
 {
+	GeoImportVersion = ESkeletalMeshGeoImportVersions::Before_Versionning;
+	SkinningImportVersion = ESkeletalMeshSkinningImportVersions::Before_Versionning;
 }
 
 
@@ -563,30 +573,41 @@ FArchive& operator<<(FArchive& Ar, FSkeletalMeshImportData& RawMesh)
 
 void FRawSkeletalMeshBulkData::Serialize(FArchive& Ar, UObject* Owner)
 {
-	if (Ar.IsLoading())
+	if (Ar.IsTransacting())
 	{
-		//Save the custom version so we can load FReductionSkeletalMeshData later
-		SerializeLoadingCustomVersionContainer = Ar.GetCustomVersions();
-		bUseSerializeLoadingCustomVersion = false;
-		const FCustomVersionContainer& RegisteredCustomVersionContainer = FCustomVersionContainer::GetRegistered();
-		const FCustomVersionArray& ArchiveCustomVersionArray = SerializeLoadingCustomVersionContainer.GetAllVersions();
-		for (const FCustomVersion& ArchiveVersion : ArchiveCustomVersionArray)
+		// If transacting, keep these members alive the other side of an undo, otherwise their values will get lost
+		SerializeLoadingCustomVersionContainer.Serialize(Ar);
+		Ar << bUseSerializeLoadingCustomVersion;
+	}
+	else
+	{
+		if (Ar.IsLoading())
 		{
-			const FCustomVersion* RegisteredVersion = RegisteredCustomVersionContainer.GetVersion(ArchiveVersion.Key);
-			if (!RegisteredVersion || RegisteredVersion->Version != ArchiveVersion.Version)
-			{
-				bUseSerializeLoadingCustomVersion = true;
-				break;
-			}
+			//Save the custom version so we can load FReductionSkeletalMeshData later
+			SerializeLoadingCustomVersionContainer = Ar.GetCustomVersions();
+			bUseSerializeLoadingCustomVersion = true;
+		}
+
+		if (Ar.IsSaving() && bUseSerializeLoadingCustomVersion == true)
+		{
+			//We need to update the FReductionSkeletalMeshData serialize version to the latest in case we save the Parent bulkdata
+			FSkeletalMeshImportData MeshImportData;
+			LoadRawMesh(MeshImportData);
+			SaveRawMesh(MeshImportData);
 		}
 	}
 
-	if (Ar.IsSaving() && bUseSerializeLoadingCustomVersion == true)
+	Ar.UsingCustomVersion(FEditorObjectVersion::GUID);
+
+	if (Ar.CustomVer(FEditorObjectVersion::GUID) >= FEditorObjectVersion::SkeletalMeshBuildRefactor)
 	{
-		//We need to update the FReductionSkeletalMeshData serialize version to the latest in case we save the Parent bulkdata
-		FSkeletalMeshImportData MeshImportData;
-		LoadRawMesh(MeshImportData);
-		SaveRawMesh(MeshImportData);
+		Ar << GeoImportVersion;
+		Ar << SkinningImportVersion;
+	}
+	else
+	{
+		GeoImportVersion = ESkeletalMeshGeoImportVersions::Before_Versionning;
+		SkinningImportVersion = ESkeletalMeshSkinningImportVersions::Before_Versionning;
 	}
 	BulkData.Serialize(Ar, Owner);
 	Ar << Guid;
@@ -597,16 +618,33 @@ void FRawSkeletalMeshBulkData::SaveRawMesh(FSkeletalMeshImportData& InMesh)
 {
 	//Saving the bulk data mean we do not need anymore the SerializeLoadingCustomVersionContainer of the parent bulk data
 	SerializeLoadingCustomVersionContainer.Empty();
-	bUseSerializeLoadingCustomVersion = false;
 
-	TArray<uint8> TempBytes;
-	FMemoryWriter Ar(TempBytes, /*bIsPersistent=*/ true);
-	Ar << InMesh;
-	BulkData.Lock(LOCK_READ_WRITE);
-	uint8* Dest = (uint8*)BulkData.Realloc(TempBytes.Num());
-	FMemory::Memcpy(Dest, TempBytes.GetData(), TempBytes.Num());
-	BulkData.Unlock();
-	FPlatformMisc::CreateGuid(Guid);
+	//Clear the bulk data before writing it
+	BulkData.RemoveBulkData();
+
+	// Get a lock on the bulk data
+	{
+		const bool bIsPersistent = true;
+		FBulkDataWriter Ar(BulkData, bIsPersistent);
+		Ar << InMesh;
+
+		// Preserve CustomVersions at save time so we can reuse the same ones when reloading direct from memory
+		SerializeLoadingCustomVersionContainer = Ar.GetCustomVersions();
+	}
+	// Unlock bulk data when we leave scope
+
+	//Create the guid from the content, this allow to use the data into the ddc key
+	FSHA1 Sha;
+	if (BulkData.GetBulkDataSize() > 0)
+	{
+		uint8* Buffer = (uint8*)BulkData.Lock(LOCK_READ_ONLY);
+		Sha.Update(Buffer, BulkData.GetBulkDataSize());
+		BulkData.Unlock();
+	}
+	Sha.Final();
+	uint32 Hash[5];
+	Sha.GetHash((uint8*)Hash);
+	Guid = FGuid(Hash[0] ^ Hash[4], Hash[1], Hash[2], Hash[3]);
 }
 
 void FRawSkeletalMeshBulkData::LoadRawMesh(FSkeletalMeshImportData& OutMesh)
@@ -614,17 +652,17 @@ void FRawSkeletalMeshBulkData::LoadRawMesh(FSkeletalMeshImportData& OutMesh)
 	OutMesh.Empty();
 	if (BulkData.GetElementCount() > 0)
 	{
-		FBufferReader Ar(
-			BulkData.Lock(LOCK_READ_ONLY), BulkData.GetElementCount(),
-			/*bInFreeOnClose=*/ false, /*bIsPersistent=*/ true
-		);
-		//If we load the bulkdata content and the bulk was using customVersion, we must set the same until we save the reduction data
-		if (bUseSerializeLoadingCustomVersion)
+		// Get a lock on the bulk data
 		{
+			const bool bIsPersistent = true;
+			FBulkDataReader Ar(BulkData, bIsPersistent);
+
+			// Propagate the custom version information from the package to the bulk data, so that the MeshDescription
+			// is serialized with the same versioning.
 			Ar.SetCustomVersions(SerializeLoadingCustomVersionContainer);
+			Ar << OutMesh;
 		}
-		Ar << OutMesh;
-		BulkData.Unlock();
+		// Unlock bulk data when we leave scope
 	}
 }
 
diff --git a/Engine/Source/Runtime/Engine/Public/Rendering/SkeletalMeshLODImporterData.h b/Engine/Source/Runtime/Engine/Public/Rendering/SkeletalMeshLODImporterData.h
--- a/Engine/Source/Runtime/Engine/Public/Rendering/SkeletalMeshLODImporterData.h
+++ b/Engine/Source/Runtime/Engine/Public/Rendering/SkeletalMeshLODImporterData.h
@@ -15,6 +15,38 @@
 
 class FSkeletalMeshLODModel;
 
+#endif
+
+//////////////////////////////////////////////////////////////////////////
+//uenum class cannot be inside a preprocessor like #if WITH_EDITOR
+
+UENUM()
+enum class ESkeletalMeshGeoImportVersions : uint8
+{
+	Before_Versionning = 0,
+	SkeletalMeshBuildRefactor,
+
+	// -----<new versions can be added above this line>-------------------------------------------------
+	VersionPlusOne,
+	LatestVersion = VersionPlusOne - 1
+};
+
+UENUM()
+enum class ESkeletalMeshSkinningImportVersions : uint8
+{
+	Before_Versionning = 0,
+	SkeletalMeshBuildRefactor,
+
+	// -----<new versions can be added above this line>-------------------------------------------------
+	VersionPlusOne,
+	LatestVersion = VersionPlusOne - 1
+};
+
+// End of enum declaration
+//////////////////////////////////////////////////////////////////////////
+
+#if WITH_EDITOR
+
 namespace SkeletalMeshImportData
 {
 	struct FMeshWedge
@@ -391,10 +423,10 @@ public:
 	ENGINE_API void Serialize(class FArchive& Ar, class UObject* Owner);
 
 	/** Store a new raw mesh in the bulk data. */
-	ENGINE_API void SaveReductionData(FSkeletalMeshLODModel& BaseLODModel, TMap<FString, TArray<FMorphTargetDelta>>& BaseLODMorphTargetData);
+	ENGINE_API void SaveReductionData(FSkeletalMeshLODModel& BaseLODModel, TMap<FString, TArray<FMorphTargetDelta>>& BaseLODMorphTargetData, UObject* Owner);
 
 	/** Load the raw mesh from bulk data. */
-	ENGINE_API void LoadReductionData(FSkeletalMeshLODModel& BaseLODModel, TMap<FString, TArray<FMorphTargetDelta>>& BaseLODMorphTargetData);
+	ENGINE_API void LoadReductionData(FSkeletalMeshLODModel& BaseLODModel, TMap<FString, TArray<FMorphTargetDelta>>& BaseLODMorphTargetData, UObject* Owner);
 
 	ENGINE_API FByteBulkData& GetBulkData() { return BulkData; }
 
@@ -421,6 +453,18 @@ class FRawSkeletalMeshBulkData
 	bool bUseSerializeLoadingCustomVersion = false;
 
 public:
+	/*
+	 * The last geo imported version, we use this flag to know if we have some data or not.
+	 * This flag must be updated every time we import a new geometry
+	 */
+	ESkeletalMeshGeoImportVersions GeoImportVersion;
+
+	/*
+	 * The last skinning imported version, we use this flag to know if we have some data or not.
+	 * This flag must be updated every time we import the skinning
+	 */
+	ESkeletalMeshSkinningImportVersions SkinningImportVersion;
+
 	/** Default constructor. */
 	ENGINE_API FRawSkeletalMeshBulkData();
 
@@ -446,6 +490,13 @@ public:
 
 	/** Returns true if no bulk data is available for this mesh. */
 	FORCEINLINE bool IsEmpty() const { return BulkData.GetBulkDataSize() == 0; }
+
+	/** Returns true if the last import version is enough to use the new build system. WE cannot rebuild asset if we did not previously store the data*/
+	ENGINE_API bool IsBuildDataAvailable() const
+	{
+		return GeoImportVersion >= ESkeletalMeshGeoImportVersions::SkeletalMeshBuildRefactor &&
+			SkinningImportVersion >= ESkeletalMeshSkinningImportVersions::SkeletalMeshBuildRefactor;
+	}
 };
 
 namespace FWedgePositionHelper
diff --git a/Engine/Source/Runtime/Engine/Public/Rendering/SkeletalMeshLODModel.h b/Engine/Source/Runtime/Engine/Public/Rendering/SkeletalMeshLODModel.h
--- a/Engine/Source/Runtime/Engine/Public/Rendering/SkeletalMeshLODModel.h
+++ b/Engine/Source/Runtime/Engine/Public/Rendering/SkeletalMeshLODModel.h
@@ -127,6 +127,24 @@ struct FSkelMeshSection
 	 */
 	int32 GenerateUpToLodIndex;
 
+	/*
+	 * This represent the original section index in the imported data. The original data is chunk per material,
+	 * we use this index to store user section modification. The user cannot change a BONE chunked section data,
+	 * since the BONE chunk can be per-platform. Do not use this value to index the Sections array, only the user
+	 * section data should be index by this value.
+	 */
+	int32 OriginalDataSectionIndex;
+
+	/*
+	 * If this section was produce because of BONE chunking, the parent section index will be valid.
+	 * If the section is not the result of skin vertex chunking, this value will be INDEX_NONE.
+	 * Use this value to know if the section was BONE chunked:
+	 * if(ChunkedParentSectionIndex != INDEX_NONE) will be true if the section is BONE chunked
+	 */
+	int32 ChunkedParentSectionIndex;
+
+
+
 	FSkelMeshSection()
 		: MaterialIndex(0)
 		, BaseIndex(0)
@@ -141,7 +159,9 @@ struct FSkelMeshSection
 		, MaxBoneInfluences(4)
 		, CorrespondClothAssetIndex(INDEX_NONE)
 		, bDisabled(false)
-		, GenerateUpToLodIndex(-1)
+		, GenerateUpToLodIndex(INDEX_NONE)
+		, OriginalDataSectionIndex(INDEX_NONE)
+		, ChunkedParentSectionIndex(INDEX_NONE)
 	{}
 
 
@@ -185,6 +205,64 @@ struct FSkelMeshSection
 	friend FArchive& operator<<(FArchive& Ar, FSkelMeshSection& S);
 };
 
+/**
+* Structure containing all the section data a user can change.
+*
+* Some section data also impact dependent generated LOD, those member should be add to the DDC Key
+* and trig a rebuild if they are change.
+*/
+struct FSkelMeshSourceSectionUserData
+{
+	/** This section will recompute tangent in runtime */
+	bool bRecomputeTangent;
+
+	/** This section will cast shadow */
+	bool bCastShadow;
+
+	// INDEX_NONE if not set
+	int16 CorrespondClothAssetIndex;
+
+	/** Clothing data for this section, clothing is only present if ClothingData.IsValid() returns true */
+	FClothingSectionData ClothingData;
+
+
+	//////////////////////////////////////////////////////////////////////////
+	//Skeletalmesh DDC key members, Add sections member that impact generated skel mesh here
+
+	/** If disabled, we won't render this section */
+	bool bDisabled;
+
+	/*
+	 * The LOD index at which any generated lower quality LODs will include this section.
+	 * A value of -1 mean the section will always be include when generating a LOD
+	 */
+	int32 GenerateUpToLodIndex;
+
+	// End DDC members
+	//////////////////////////////////////////////////////////////////////////
+
+
+
+	FSkelMeshSourceSectionUserData()
+		: bRecomputeTangent(false)
+		, bCastShadow(true)
+		, CorrespondClothAssetIndex(INDEX_NONE)
+		, bDisabled(false)
+		, GenerateUpToLodIndex(INDEX_NONE)
+	{}
+
+	/**
+	* @return TRUE if we have cloth data for this section
+	*/
+	FORCEINLINE bool HasClothingData() const
+	{
+		return (ClothingData.AssetGuid.IsValid());
+	}
+
+	// Serialization.
+	friend FArchive& operator<<(FArchive& Ar, FSkelMeshSourceSectionUserData& S);
+};
+
 /**
 * All data to define a certain LOD model for a skeletal mesh.
 */
@@ -194,6 +272,14 @@ public:
 	/** Sections. */
 	TArray<FSkelMeshSection> Sections;
 
+	/*
+	 * When user change section data in the UI, we store it here to be able to regenerate the changes
+	 * Note: the key (int32) is the original imported section data, because of BONE chunk the size of
+	 * this array is not the same as the Sections array. Use the section's OriginalDataSectionIndex to
+	 * index it.
+	 */
+	TMap<int32, FSkelMeshSourceSectionUserData> UserSectionsData;
+
 	uint32						NumVertices;
 	/** The number of unique texture coordinate sets in this lod */
 	uint32						NumTexCoords;
@@ -234,6 +320,7 @@ public:
 		: NumVertices(0)
 		, NumTexCoords(0)
 		, MaxImportVertex(-1)
+		, BuildStringID(TEXT(""))
 	{
 	}
 
@@ -269,7 +356,7 @@ public:
 	/**
 	* @return true if any chunks have cloth data.
 	*/
-	bool HasClothData() const;
+	ENGINE_API bool HasClothData() const;
 
 	ENGINE_API int32 NumNonClothingSections() const;
 
@@ -286,6 +373,21 @@ public:
 	ENGINE_API bool DoSectionsNeedExtraBoneInfluences() const;
 
 	void GetResourceSizeEx(FResourceSizeEx& CumulativeResourceSize) const;
+
+	/**
+	* Make sure user section data is present for every original sections
+	*/
+	ENGINE_API void SyncronizeUserSectionsDataArray(bool bResetNonUsedSection = false);
+
+	//Temporary build String ID
+	//We use this string to store the LOD model data so we can know if the LOD need to be rebuild
+	//This GUID is set when we Cache the render data (build function)
+	mutable FString BuildStringID;
+
+	/**
+	* Build a derive data key with the user section data (UserSectionsData) and the original bulk data
+	*/
+	ENGINE_API FString GetLODModelDeriveDataKey() const;
 };
 
 #endif // WITH_EDITOR
\ No newline at end of file
diff --git a/Engine/Source/Runtime/Engine/Public/Rendering/SkeletalMeshModel.h b/Engine/Source/Runtime/Engine/Public/Rendering/SkeletalMeshModel.h
--- a/Engine/Source/Runtime/Engine/Public/Rendering/SkeletalMeshModel.h
+++ b/Engine/Source/Runtime/Engine/Public/Rendering/SkeletalMeshModel.h
@@ -42,6 +42,10 @@ public:
 	/** Get current GUID Id as a string, for DDC key */
 	FString GetIdString() const;
 
+	void SyncronizeLODUserSectionsData();
+
+	FString GetLODModelIdString() const;
+
 	ENGINE_API void EmptyOriginalReductionSourceMeshData();
 	/* When user reduce an imported LOD with himself (BaseLOD == TargetLOD), we need to store some imported model data so we can reduce again from the same data.*/
 	TArray<FReductionBaseSkeletalMeshBulkData*> OriginalReductionSourceMeshData;
diff --git a/Engine/Source/Runtime/Engine/Public/SkeletalMeshTypes.h b/Engine/Source/Runtime/Engine/Public/SkeletalMeshTypes.h
--- a/Engine/Source/Runtime/Engine/Public/SkeletalMeshTypes.h
+++ b/Engine/Source/Runtime/Engine/Public/SkeletalMeshTypes.h
@@ -5,6 +5,7 @@
 #include "CoreMinimal.h"
 #include "PrimitiveSceneProxy.h"
 #include "Materials/MaterialInterface.h"
+#include "ComponentReregisterContext.h"
 
 class FMaterialRenderProxy;
 class FMeshElementCollector;
@@ -122,7 +123,8 @@ struct ESkeletalMeshVertexFlags
 		None = 0x0,
 		UseFullPrecisionUVs = 0x1,
 		HasVertexColors = 0x2,
-		UseHighPrecisionTangentBasis = 0x4
+		UseHighPrecisionTangentBasis = 0x4,
+		BuildAdjacencyIndexBuffer = 0x8
 	};
 };
 
@@ -414,3 +416,35 @@ private:
 	bool bRefreshBounds;
 };
 
+#if WITH_EDITOR
+
+//Helper to scope skeletal mesh post edit change.
+class ENGINE_API FScopedSkeletalMeshPostEditChange
+{
+public:
+	/*
+	 * This constructor increment the skeletal mesh PostEditChangeStackCounter. If the stack counter is zero before the increment
+	 * the skeletal mesh component will be unregister from the world. The component will also release there rendering resources.
+	 * Parameters:
+	 * @param InbCallPostEditChange - if we are the first scope PostEditChange will be called.
+	 * @param InbReregisterComponents - if we are the first scope we will re register component from world and also component render data.
+	 */
+	FScopedSkeletalMeshPostEditChange(USkeletalMesh* InSkeletalMesh, bool InbCallPostEditChange = true, bool InbReregisterComponents = true);
+
+	/*
+	 * This destructor decrement the skeletal mesh PostEditChangeStackCounter. If the stack counter is zero after the decrement,
+	 * the skeletal mesh PostEditChange will be call. The component will also be register to the world and there render data resources will be rebuild.
+	 */
+	~FScopedSkeletalMeshPostEditChange();
+
+	void SetSkeletalMesh(USkeletalMesh* InSkeletalMesh);
+
+private:
+	USkeletalMesh* SkeletalMesh;
+	bool bReregisterComponents;
+	bool bCallPostEditChange;
+	FSkinnedMeshComponentRecreateRenderStateContext* RecreateExistingRenderStateContext;
+	TIndirectArray<FComponentReregisterContext> ComponentReregisterContexts;
+};
+
+#endif
diff --git a/Engine/Source/Runtime/MeshDescription/MeshDescription.Build.cs b/Engine/Source/Runtime/MeshDescription/MeshDescription.Build.cs
--- a/Engine/Source/Runtime/MeshDescription/MeshDescription.Build.cs
+++ b/Engine/Source/Runtime/MeshDescription/MeshDescription.Build.cs
@@ -9,11 +9,11 @@ namespace UnrealBuildTool.Rules
             PrivateIncludePaths.Add("Runtime/MeshDescription/Private");
             PublicIncludePaths.Add("Runtime/MeshDescription/Public");
 
-			PrivateDependencyModuleNames.AddRange(
+			PublicDependencyModuleNames.AddRange(
 				new string[]
 				{
 					"Core",
-					"CoreUObject",
+					"CoreUObject"
 				}
 			);
 		}
diff --git a/Engine/Source/Runtime/MeshDescription/Private/MeshAttributes.cpp b/Engine/Source/Runtime/MeshDescription/Private/MeshAttributes.cpp
--- a/Engine/Source/Runtime/MeshDescription/Private/MeshAttributes.cpp
+++ b/Engine/Source/Runtime/MeshDescription/Private/MeshAttributes.cpp
@@ -4,26 +4,12 @@
 
 namespace MeshAttribute
 {
-	const FName Vertex::Position( "Position ");
-	const FName Vertex::CornerSharpness( "CornerSharpness" );
-
-	const FName VertexInstance::TextureCoordinate( "TextureCoordinate" );
-	const FName VertexInstance::Normal( "Normal" );
-	const FName VertexInstance::Tangent( "Tangent" );
-	const FName VertexInstance::BinormalSign( "BinormalSign" );
-	const FName VertexInstance::Color( "Color" );
-
-	const FName Edge::IsHard( "IsHard" );
-	const FName Edge::IsUVSeam( "IsUVSeam" );
-	const FName Edge::CreaseSharpness( "CreaseSharpness" );
+	const FName Vertex::Position("Position ");
+}
 
-	const FName Polygon::Normal( "Normal" );
-	const FName Polygon::Tangent( "Tangent" );
-	const FName Polygon::Binormal( "Binormal" );
-	const FName Polygon::Center( "Center" );
 
-	const FName PolygonGroup::ImportedMaterialSlotName( "ImportedMaterialSlotName" );
-	const FName PolygonGroup::EnableCollision( "EnableCollision" );
-	const FName PolygonGroup::CastShadow( "CastShadow" );
+void FMeshAttributes::Register()
+{
+	// Nothing to do here: Vertex positions are already registered by the FMeshDescription constructor
 }
 
diff --git a/Engine/Source/Runtime/MeshDescription/Private/MeshDescription.cpp b/Engine/Source/Runtime/MeshDescription/Private/MeshDescription.cpp
--- a/Engine/Source/Runtime/MeshDescription/Private/MeshDescription.cpp
+++ b/Engine/Source/Runtime/MeshDescription/Private/MeshDescription.cpp
@@ -23,6 +23,13 @@ void UDEPRECATED_MeshDescription::Serialize( FArchive& Ar )
 
 PRAGMA_DISABLE_DEPRECATION_WARNINGS
 
+FMeshDescription::FMeshDescription()
+{
+	// Minimal requirement is that vertices have a Position attribute
+	VertexAttributesSet.RegisterAttribute(MeshAttribute::Vertex::Position, 1, FVector::ZeroVector, EMeshAttributeFlags::Lerpable);
+}
+
+
 void FMeshDescription::Serialize(FArchive& Ar)
 {
 	Ar.UsingCustomVersion(FReleaseObjectVersion::GUID);
@@ -1229,6 +1236,39 @@ void FMeshDescription::ComputePolygonTriangulation(const FPolygonID PolygonID, T
 }
 
 
+FBoxSphereBounds FMeshDescription::GetBounds() const
+{
+	FBoxSphereBounds BoundingBoxAndSphere;
+
+	FBox BoundingBox;
+	BoundingBox.Init();
+
+	TVertexAttributesConstRef<FVector> VertexPositions = VertexAttributes().GetAttributesRef<FVector>(MeshAttribute::Vertex::Position);
+
+	for (const FVertexID VertexID : Vertices().GetElementIDs())
+	{
+		if (!IsVertexOrphaned(VertexID))
+		{
+			BoundingBox += VertexPositions[VertexID];
+		}
+	}
+
+	BoundingBox.GetCenterAndExtents(BoundingBoxAndSphere.Origin, BoundingBoxAndSphere.BoxExtent);
+
+	// Calculate the bounding sphere, using the center of the bounding box as the origin.
+	BoundingBoxAndSphere.SphereRadius = 0.0f;
+
+	for (const FVertexID VertexID : Vertices().GetElementIDs())
+	{
+		if (!IsVertexOrphaned(VertexID))
+		{
+			BoundingBoxAndSphere.SphereRadius = FMath::Max((VertexPositions[VertexID] - BoundingBoxAndSphere.Origin).Size(), BoundingBoxAndSphere.SphereRadius);
+		}
+	}
+
+	return BoundingBoxAndSphere;
+}
+
 void FMeshDescription::TriangulateMesh()
 {
 	// Perform triangulation directly into mesh polygons
@@ -1239,6 +1279,47 @@ void FMeshDescription::TriangulateMesh()
 }
 
 
+namespace MeshAttribute_
+{
+	namespace Vertex
+	{
+		const FName CornerSharpness("CornerSharpness");
+	}
+
+	namespace VertexInstance
+	{
+		const FName TextureCoordinate("TextureCoordinate");
+		const FName Normal("Normal");
+		const FName Tangent("Tangent");
+		const FName BinormalSign("BinormalSign");
+		const FName Color("Color");
+	}
+
+	namespace Edge
+	{
+		const FName IsHard("IsHard");
+		const FName IsUVSeam("IsUVSeam");
+		const FName CreaseSharpness("CreaseSharpness");
+	}
+
+	namespace Polygon
+	{
+		const FName Normal("Normal");
+		const FName Tangent("Tangent");
+		const FName Binormal("Binormal");
+		const FName Center("Center");
+	}
+
+	namespace PolygonGroup
+	{
+		const FName ImportedMaterialSlotName("ImportedMaterialSlotName");
+		const FName EnableCollision("EnableCollision");
+		const FName CastShadow("CastShadow");
+	}
+}
+
+
+
 bool FMeshDescription::ComputePolygonTangentsAndNormals(const FPolygonID PolygonID
 	, float ComparisonThreshold
 	, const TVertexAttributesRef<FVector> VertexPositions
@@ -1308,11 +1389,11 @@ bool FMeshDescription::ComputePolygonTangentsAndNormals(const FPolygonID Polygon
 void FMeshDescription::ComputePolygonTangentsAndNormals(const TArray<FPolygonID>& PolygonIDs, float ComparisonThreshold)
 {
 	const TVertexAttributesRef<FVector> VertexPositions = VertexAttributes().GetAttributesRef<FVector>(MeshAttribute::Vertex::Position);
-	const TVertexInstanceAttributesRef<FVector2D> VertexUVs = VertexInstanceAttributes().GetAttributesRef<FVector2D>(MeshAttribute::VertexInstance::TextureCoordinate);
-	TPolygonAttributesRef<FVector> PolygonNormals = PolygonAttributes().GetAttributesRef<FVector>(MeshAttribute::Polygon::Normal);
-	TPolygonAttributesRef<FVector> PolygonTangents = PolygonAttributes().GetAttributesRef<FVector>(MeshAttribute::Polygon::Tangent);
-	TPolygonAttributesRef<FVector> PolygonBinormals = PolygonAttributes().GetAttributesRef<FVector>(MeshAttribute::Polygon::Binormal);
-	TPolygonAttributesRef<FVector> PolygonCenters = PolygonAttributes().GetAttributesRef<FVector>(MeshAttribute::Polygon::Center);
+	const TVertexInstanceAttributesRef<FVector2D> VertexUVs = VertexInstanceAttributes().GetAttributesRef<FVector2D>(MeshAttribute_::VertexInstance::TextureCoordinate);
+	TPolygonAttributesRef<FVector> PolygonNormals = PolygonAttributes().GetAttributesRef<FVector>(MeshAttribute_::Polygon::Normal);
+	TPolygonAttributesRef<FVector> PolygonTangents = PolygonAttributes().GetAttributesRef<FVector>(MeshAttribute_::Polygon::Tangent);
+	TPolygonAttributesRef<FVector> PolygonBinormals = PolygonAttributes().GetAttributesRef<FVector>(MeshAttribute_::Polygon::Binormal);
+	TPolygonAttributesRef<FVector> PolygonCenters = PolygonAttributes().GetAttributesRef<FVector>(MeshAttribute_::Polygon::Center);
 
 	TArray<FPolygonID> DegeneratePolygonIDs;
 	for (const FPolygonID PolygonID : PolygonIDs)
@@ -1374,7 +1455,7 @@ void FMeshDescription::GetConnectedSoftEdges(const FVertexID VertexID, TArray<FE
 {
 	OutConnectedSoftEdges.Reset();
 
-	TEdgeAttributesConstRef<bool> EdgeHardnesses = EdgeAttributes().GetAttributesRef<bool>(MeshAttribute::Edge::IsHard);
+	TEdgeAttributesConstRef<bool> EdgeHardnesses = EdgeAttributes().GetAttributesRef<bool>(MeshAttribute_::Edge::IsHard);
 	for (const FEdgeID ConnectedEdgeID : VertexArray[VertexID].ConnectedEdgeIDs)
 	{
 		if (!EdgeHardnesses[ConnectedEdgeID])
@@ -1608,13 +1689,13 @@ void FMeshDescription::ComputeTangentsAndNormals(const FVertexInstanceID VertexI
 
 void FMeshDescription::ComputeTangentsAndNormals(const TArray<FVertexInstanceID>& VertexInstanceIDs, EComputeNTBsOptions ComputeNTBsOptions)
 {
-	const TPolygonAttributesRef<FVector> PolygonNormals = PolygonAttributes().GetAttributesRef<FVector>(MeshAttribute::Polygon::Normal);
-	const TPolygonAttributesRef<FVector> PolygonTangents = PolygonAttributes().GetAttributesRef<FVector>(MeshAttribute::Polygon::Tangent);
-	const TPolygonAttributesRef<FVector> PolygonBinormals = PolygonAttributes().GetAttributesRef<FVector>(MeshAttribute::Polygon::Binormal);
+	const TPolygonAttributesRef<FVector> PolygonNormals = PolygonAttributes().GetAttributesRef<FVector>(MeshAttribute_::Polygon::Normal);
+	const TPolygonAttributesRef<FVector> PolygonTangents = PolygonAttributes().GetAttributesRef<FVector>(MeshAttribute_::Polygon::Tangent);
+	const TPolygonAttributesRef<FVector> PolygonBinormals = PolygonAttributes().GetAttributesRef<FVector>(MeshAttribute_::Polygon::Binormal);
 
-	TVertexInstanceAttributesRef<FVector> VertexNormals = VertexInstanceAttributes().GetAttributesRef<FVector>(MeshAttribute::VertexInstance::Normal);
-	TVertexInstanceAttributesRef<FVector> VertexTangents = VertexInstanceAttributes().GetAttributesRef<FVector>(MeshAttribute::VertexInstance::Tangent);
-	TVertexInstanceAttributesRef<float> VertexBinormalSigns = VertexInstanceAttributes().GetAttributesRef<float>(MeshAttribute::VertexInstance::BinormalSign);
+	TVertexInstanceAttributesRef<FVector> VertexNormals = VertexInstanceAttributes().GetAttributesRef<FVector>(MeshAttribute_::VertexInstance::Normal);
+	TVertexInstanceAttributesRef<FVector> VertexTangents = VertexInstanceAttributes().GetAttributesRef<FVector>(MeshAttribute_::VertexInstance::Tangent);
+	TVertexInstanceAttributesRef<float> VertexBinormalSigns = VertexInstanceAttributes().GetAttributesRef<float>(MeshAttribute_::VertexInstance::BinormalSign);
 
 	for (const FVertexInstanceID VertexInstanceID : VertexInstanceIDs)
 	{
@@ -1625,13 +1706,13 @@ void FMeshDescription::ComputeTangentsAndNormals(const TArray<FVertexInstanceID>
 
 void FMeshDescription::ComputeTangentsAndNormals(EComputeNTBsOptions ComputeNTBsOptions)
 {
-	const TPolygonAttributesRef<FVector> PolygonNormals = PolygonAttributes().GetAttributesRef<FVector>(MeshAttribute::Polygon::Normal);
-	const TPolygonAttributesRef<FVector> PolygonTangents = PolygonAttributes().GetAttributesRef<FVector>(MeshAttribute::Polygon::Tangent);
-	const TPolygonAttributesRef<FVector> PolygonBinormals = PolygonAttributes().GetAttributesRef<FVector>(MeshAttribute::Polygon::Binormal);
+	const TPolygonAttributesRef<FVector> PolygonNormals = PolygonAttributes().GetAttributesRef<FVector>(MeshAttribute_::Polygon::Normal);
+	const TPolygonAttributesRef<FVector> PolygonTangents = PolygonAttributes().GetAttributesRef<FVector>(MeshAttribute_::Polygon::Tangent);
+	const TPolygonAttributesRef<FVector> PolygonBinormals = PolygonAttributes().GetAttributesRef<FVector>(MeshAttribute_::Polygon::Binormal);
 
-	TVertexInstanceAttributesRef<FVector> VertexNormals = VertexInstanceAttributes().GetAttributesRef<FVector>(MeshAttribute::VertexInstance::Normal);
-	TVertexInstanceAttributesRef<FVector> VertexTangents = VertexInstanceAttributes().GetAttributesRef<FVector>(MeshAttribute::VertexInstance::Tangent);
-	TVertexInstanceAttributesRef<float> VertexBinormalSigns = VertexInstanceAttributes().GetAttributesRef<float>(MeshAttribute::VertexInstance::BinormalSign);
+	TVertexInstanceAttributesRef<FVector> VertexNormals = VertexInstanceAttributes().GetAttributesRef<FVector>(MeshAttribute_::VertexInstance::Normal);
+	TVertexInstanceAttributesRef<FVector> VertexTangents = VertexInstanceAttributes().GetAttributesRef<FVector>(MeshAttribute_::VertexInstance::Tangent);
+	TVertexInstanceAttributesRef<float> VertexBinormalSigns = VertexInstanceAttributes().GetAttributesRef<float>(MeshAttribute_::VertexInstance::BinormalSign);
 
 	for (const FVertexInstanceID VertexInstanceID : VertexInstances().GetElementIDs())
 	{
@@ -1642,8 +1723,8 @@ void FMeshDescription::ComputeTangentsAndNormals(EComputeNTBsOptions ComputeNTBs
 
 void FMeshDescription::DetermineEdgeHardnessesFromVertexInstanceNormals( const float Tolerance )
 {
-	const TVertexInstanceAttributesRef<FVector> VertexNormals = VertexInstanceAttributes().GetAttributesRef<FVector>( MeshAttribute::VertexInstance::Normal );
-	TEdgeAttributesRef<bool> EdgeHardnesses = EdgeAttributes().GetAttributesRef<bool>( MeshAttribute::Edge::IsHard );
+	const TVertexInstanceAttributesRef<FVector> VertexNormals = VertexInstanceAttributes().GetAttributesRef<FVector>( MeshAttribute_::VertexInstance::Normal );
+	TEdgeAttributesRef<bool> EdgeHardnesses = EdgeAttributes().GetAttributesRef<bool>( MeshAttribute_::Edge::IsHard );
 
 	// Holds unique vertex instance IDs for a given edge vertex
 	// @todo: use TMemStackAllocator or similar to avoid expensive allocations
@@ -1706,8 +1787,8 @@ void FMeshDescription::DetermineEdgeHardnessesFromVertexInstanceNormals( const f
 
 void FMeshDescription::DetermineUVSeamsFromUVs( const int32 UVIndex, const float Tolerance )
 {
-	const TVertexInstanceAttributesRef<FVector2D> VertexUVs = VertexInstanceAttributes().GetAttributesRef<FVector2D>( MeshAttribute::VertexInstance::TextureCoordinate );
-	TEdgeAttributesRef<bool> EdgeUVSeams = EdgeAttributes().GetAttributesRef<bool>( MeshAttribute::Edge::IsUVSeam );
+	const TVertexInstanceAttributesRef<FVector2D> VertexUVs = VertexInstanceAttributes().GetAttributesRef<FVector2D>( MeshAttribute_::VertexInstance::TextureCoordinate );
+	TEdgeAttributesRef<bool> EdgeUVSeams = EdgeAttributes().GetAttributesRef<bool>( MeshAttribute_::Edge::IsUVSeam );
 
 	// Holds unique vertex instance IDs for a given edge vertex
 	// @todo: use TMemStackAllocator or similar to avoid expensive allocations
@@ -1769,7 +1850,7 @@ void FMeshDescription::DetermineUVSeamsFromUVs( const int32 UVIndex, const float
 
 void FMeshDescription::GetPolygonsInSameChartAsPolygon( const FPolygonID PolygonID, TArray<FPolygonID>& OutPolygonIDs )
 {
-	const TEdgeAttributesRef<bool> EdgeUVSeams = EdgeAttributes().GetAttributesRef<bool>( MeshAttribute::Edge::IsUVSeam );
+	const TEdgeAttributesRef<bool> EdgeUVSeams = EdgeAttributes().GetAttributesRef<bool>( MeshAttribute_::Edge::IsUVSeam );
 	const int32 NumPolygons = Polygons().Num();
 
 	// This holds the results - all polygon IDs which are in the same UV chart
@@ -1871,7 +1952,7 @@ void FMeshDescription::ReverseAllPolygonFacing()
 
 void FMeshDescription::RemapPolygonGroups(const TMap<FPolygonGroupID, FPolygonGroupID>& Remap)
 {
-	TPolygonGroupAttributesRef<FName> PolygonGroupNames = PolygonGroupAttributes().GetAttributesRef<FName>(MeshAttribute::PolygonGroup::ImportedMaterialSlotName);
+	TPolygonGroupAttributesRef<FName> PolygonGroupNames = PolygonGroupAttributes().GetAttributesRef<FName>(MeshAttribute_::PolygonGroup::ImportedMaterialSlotName);
 
 	struct FOldPolygonGroupData
 	{
diff --git a/Engine/Source/Runtime/MeshDescription/Private/MeshDescriptionBase.cpp b/Engine/Source/Runtime/MeshDescription/Private/MeshDescriptionBase.cpp
new file mode 100644
--- /dev/null
+++ b/Engine/Source/Runtime/MeshDescription/Private/MeshDescriptionBase.cpp
@@ -0,0 +1,992 @@
+// Copyright 1998-2019 Epic Games, Inc. All Rights Reserved.
+
+#include "MeshDescriptionBase.h"
+#include "Algo/Copy.h"
+
+
+
+void UMeshDescriptionBase::Serialize(FArchive& Ar)
+{
+	Super::Serialize(Ar);
+}
+
+void UMeshDescriptionBase::RegisterAttributes()
+{
+	RequiredAttributes = MakeUnique<FMeshAttributes>(MeshDescription);
+	RequiredAttributes->Register();
+}
+
+void UMeshDescriptionBase::Reset()
+{
+	MeshDescription = FMeshDescription();
+	RegisterAttributes();
+}
+
+void UMeshDescriptionBase::Empty()
+{
+	MeshDescription.Empty();
+}
+
+bool UMeshDescriptionBase::IsEmpty() const
+{
+	return MeshDescription.IsEmpty();
+}
+
+void UMeshDescriptionBase::ReserveNewVertices(int32 NumberOfNewVertices)
+{
+	MeshDescription.ReserveNewVertices(NumberOfNewVertices);
+}
+
+FVertexID UMeshDescriptionBase::CreateVertex()
+{
+	return MeshDescription.CreateVertex();
+}
+
+void UMeshDescriptionBase::CreateVertexWithID(FVertexID VertexID)
+{
+	if (!MeshDescription.IsVertexValid(VertexID))
+	{
+		MeshDescription.CreateVertexWithID(VertexID);
+	}
+	else
+	{
+		UE_LOG(LogMeshDescription, Warning, TEXT("CreateVertexWithID: VertexID %d already exists."), VertexID.GetValue());
+	}
+}
+
+void UMeshDescriptionBase::DeleteVertex(FVertexID VertexID)
+{
+	if (MeshDescription.IsVertexValid(VertexID))
+	{
+		MeshDescription.DeleteVertex(VertexID);
+	}
+	else
+	{
+		UE_LOG(LogMeshDescription, Warning, TEXT("DeleteVertex: VertexID %d doesn't exist."), VertexID.GetValue());
+	}
+}
+
+bool UMeshDescriptionBase::IsVertexValid(FVertexID VertexID) const
+{
+	return MeshDescription.IsVertexValid(VertexID);
+}
+
+void UMeshDescriptionBase::ReserveNewVertexInstances(int32 NumberOfNewVertexInstances)
+{
+	MeshDescription.ReserveNewEdges(NumberOfNewVertexInstances);
+}
+
+FVertexInstanceID UMeshDescriptionBase::CreateVertexInstance(FVertexID VertexID)
+{
+	if (MeshDescription.IsVertexValid(VertexID))
+	{
+		return MeshDescription.CreateVertexInstance(VertexID);
+	}
+	else
+	{
+		UE_LOG(LogMeshDescription, Warning, TEXT("CreateVertexInstance: VertexID %d doesn't exist."), VertexID.GetValue());
+		return FVertexInstanceID::Invalid;
+	}
+}
+
+void UMeshDescriptionBase::CreateVertexInstanceWithID(FVertexInstanceID VertexInstanceID, FVertexID VertexID)
+{
+	if (MeshDescription.IsVertexValid(VertexID))
+	{
+		if (!MeshDescription.IsVertexInstanceValid(VertexInstanceID))
+		{
+			MeshDescription.CreateVertexInstanceWithID(VertexInstanceID, VertexID);
+		}
+		else
+		{
+			UE_LOG(LogMeshDescription, Warning, TEXT("CreateVertexInstanceWithID: VertexInstanceID %d already exists."), VertexInstanceID.GetValue());
+		}
+	}
+	else
+	{
+		UE_LOG(LogMeshDescription, Warning, TEXT("CreateVertexInstanceWithID: VertexID %d doesn't exist."), VertexID.GetValue());
+	}
+}
+
+void UMeshDescriptionBase::DeleteVertexInstance(FVertexInstanceID VertexInstanceID, TArray<FVertexID>& OrphanedVertices)
+{
+	if (MeshDescription.IsVertexInstanceValid(VertexInstanceID))
+	{
+		MeshDescription.DeleteVertexInstance(VertexInstanceID, &OrphanedVertices);
+	}
+	else
+	{
+		UE_LOG(LogMeshDescription, Warning, TEXT("DeleteVertexInstance: VertexInstanceID %d doesn't exist."), VertexInstanceID.GetValue());
+	}
+}
+
+bool UMeshDescriptionBase::IsVertexInstanceValid(FVertexInstanceID VertexInstanceID) const
+{
+	return MeshDescription.IsVertexInstanceValid(VertexInstanceID);
+}
+
+void UMeshDescriptionBase::ReserveNewEdges(int32 NumberOfNewEdges)
+{
+	MeshDescription.ReserveNewEdges(NumberOfNewEdges);
+}
+
+FEdgeID UMeshDescriptionBase::CreateEdge(FVertexID VertexID0, FVertexID VertexID1)
+{
+	if (!MeshDescription.IsVertexValid(VertexID0))
+	{
+		UE_LOG(LogMeshDescription, Warning, TEXT("CreateEdge: VertexID %d doesn't exist."), VertexID0.GetValue());
+		return FEdgeID::Invalid;
+	}
+
+	if (!MeshDescription.IsVertexValid(VertexID1))
+	{
+		UE_LOG(LogMeshDescription, Warning, TEXT("CreateEdge: VertexID %d doesn't exist."), VertexID1.GetValue());
+		return FEdgeID::Invalid;
+	}
+
+	return MeshDescription.CreateEdge(VertexID0, VertexID1);
+}
+
+void UMeshDescriptionBase::CreateEdgeWithID(FEdgeID EdgeID, FVertexID VertexID0, FVertexID VertexID1)
+{
+	if (MeshDescription.IsEdgeValid(EdgeID))
+	{
+		UE_LOG(LogMeshDescription, Warning, TEXT("CreateEdgeWithID: EdgeID %d already exists."), EdgeID.GetValue());
+		return;
+	}
+
+	if (!MeshDescription.IsVertexValid(VertexID0))
+	{
+		UE_LOG(LogMeshDescription, Warning, TEXT("CreateEdgeWithID: VertexID %d doesn't exist."), VertexID0.GetValue());
+		return;
+	}
+
+	if (!MeshDescription.IsVertexValid(VertexID1))
+	{
+		UE_LOG(LogMeshDescription, Warning, TEXT("CreateEdgeWithID: VertexID %d doesn't exist."), VertexID1.GetValue());
+		return;
+	}
+
+	MeshDescription.CreateEdgeWithID(EdgeID, VertexID0, VertexID1);
+}
+
+void UMeshDescriptionBase::DeleteEdge(FEdgeID EdgeID, TArray<FVertexID>& OrphanedVertices)
+{
+	if (MeshDescription.IsEdgeValid(EdgeID))
+	{
+		MeshDescription.DeleteEdge(EdgeID, &OrphanedVertices);
+	}
+	else
+	{
+		UE_LOG(LogMeshDescription, Warning, TEXT("DeleteEdge: EdgeID %d doesn't exist."), EdgeID.GetValue());
+	}
+}
+
+bool UMeshDescriptionBase::IsEdgeValid(FEdgeID EdgeID) const
+{
+	return MeshDescription.IsEdgeValid(EdgeID);
+}
+
+void UMeshDescriptionBase::ReserveNewTriangles(int32 NumberOfNewTriangles)
+{
+	MeshDescription.ReserveNewTriangles(NumberOfNewTriangles);
+}
+
+FTriangleID UMeshDescriptionBase::CreateTriangle(FPolygonGroupID PolygonGroupID, const TArray<FVertexInstanceID>& VertexInstanceIDs, TArray<FEdgeID>& NewEdgeIDs)
+{
+	if (!MeshDescription.IsPolygonGroupValid(PolygonGroupID))
+	{
+		UE_LOG(LogMeshDescription, Warning, TEXT("CreateTriangle: PolygonGroupID %d doesn't exist."), PolygonGroupID.GetValue());
+		return FTriangleID::Invalid;
+	}
+
+	return MeshDescription.CreateTriangle(PolygonGroupID, VertexInstanceIDs, &NewEdgeIDs);
+}
+
+void UMeshDescriptionBase::CreateTriangleWithID(FTriangleID TriangleID, FPolygonGroupID PolygonGroupID, const TArray<FVertexInstanceID>& VertexInstanceIDs, TArray<FEdgeID>& NewEdgeIDs)
+{
+	if (MeshDescription.IsTriangleValid(TriangleID))
+	{
+		UE_LOG(LogMeshDescription, Warning, TEXT("CreateTriangleWithID: TriangleID %d already exists."), TriangleID.GetValue());
+		return;
+	}
+
+	if (!MeshDescription.IsPolygonGroupValid(PolygonGroupID))
+	{
+		UE_LOG(LogMeshDescription, Warning, TEXT("CreateTriangleWithID: PolygonGroupID %d doesn't exist."), PolygonGroupID.GetValue());
+		return;
+	}
+
+	MeshDescription.CreateTriangleWithID(TriangleID, PolygonGroupID, VertexInstanceIDs, &NewEdgeIDs);
+}
+
+void UMeshDescriptionBase::DeleteTriangle(FTriangleID TriangleID, TArray<FEdgeID>& OrphanedEdges, TArray<FVertexInstanceID>& OrphanedVertexInstances, TArray<FPolygonGroupID>& OrphanedPolygonGroups)
+{
+	if (!MeshDescription.IsTriangleValid(TriangleID))
+	{
+		UE_LOG(LogMeshDescription, Warning, TEXT("DeleteTriangle: TriangleID %d doesn't exist."), TriangleID.GetValue());
+		return;
+	}
+
+	MeshDescription.DeleteTriangle(TriangleID, &OrphanedEdges, &OrphanedVertexInstances, &OrphanedPolygonGroups);
+}
+
+bool UMeshDescriptionBase::IsTriangleValid(const FTriangleID TriangleID) const
+{
+	return MeshDescription.IsTriangleValid(TriangleID);
+}
+
+void UMeshDescriptionBase::ReserveNewPolygons(const int32 NumberOfNewPolygons)
+{
+	MeshDescription.ReserveNewPolygons(NumberOfNewPolygons);
+}
+
+FPolygonID UMeshDescriptionBase::CreatePolygon(FPolygonGroupID PolygonGroupID, TArray<FVertexInstanceID>& VertexInstanceIDs, TArray<FEdgeID>& NewEdgeIDs)
+{
+	if (!MeshDescription.IsPolygonGroupValid(PolygonGroupID))
+	{
+		UE_LOG(LogMeshDescription, Warning, TEXT("CreatePolygon: PolygonGroupID %d doesn't exist."), PolygonGroupID.GetValue());
+		return FPolygonID::Invalid;
+	}
+
+	return MeshDescription.CreatePolygon(PolygonGroupID, VertexInstanceIDs, &NewEdgeIDs);
+}
+
+void UMeshDescriptionBase::CreatePolygonWithID(FPolygonID PolygonID, FPolygonGroupID PolygonGroupID, TArray<FVertexInstanceID>& VertexInstanceIDs, TArray<FEdgeID>& NewEdgeIDs)
+{
+	if (MeshDescription.IsPolygonValid(PolygonID))
+	{
+		UE_LOG(LogMeshDescription, Warning, TEXT("CreatePolygonWithID: PolygonID %d already exists."), PolygonID.GetValue());
+		return;
+	}
+
+	if (!MeshDescription.IsPolygonGroupValid(PolygonGroupID))
+	{
+		UE_LOG(LogMeshDescription, Warning, TEXT("CreatePolygonWithID: PolygonGroupID %d doesn't exist."), PolygonGroupID.GetValue());
+		return;
+	}
+
+	MeshDescription.CreatePolygonWithID(PolygonID, PolygonGroupID, VertexInstanceIDs, &NewEdgeIDs);
+}
+
+void UMeshDescriptionBase::DeletePolygon(FPolygonID PolygonID, TArray<FEdgeID>& OrphanedEdges, TArray<FVertexInstanceID>& OrphanedVertexInstances, TArray<FPolygonGroupID>& OrphanedPolygonGroups)
+{
+	if (!MeshDescription.IsPolygonValid(PolygonID))
+	{
+		UE_LOG(LogMeshDescription, Warning, TEXT("DeletePolygon: PolygonID %d doesn't exist."), PolygonID.GetValue());
+		return;
+	}
+
+	MeshDescription.DeletePolygon(PolygonID, &OrphanedEdges, &OrphanedVertexInstances, &OrphanedPolygonGroups);
+}
+
+bool UMeshDescriptionBase::IsPolygonValid(FPolygonID PolygonID) const
+{
+	return MeshDescription.IsPolygonValid(PolygonID);
+}
+
+void UMeshDescriptionBase::ReserveNewPolygonGroups(int32 NumberOfNewPolygonGroups)
+{
+	MeshDescription.ReserveNewPolygonGroups(NumberOfNewPolygonGroups);
+}
+
+FPolygonGroupID UMeshDescriptionBase::CreatePolygonGroup()
+{
+	return MeshDescription.CreatePolygonGroup();
+}
+
+void UMeshDescriptionBase::CreatePolygonGroupWithID(FPolygonGroupID PolygonGroupID)
+{
+	if (MeshDescription.IsPolygonGroupValid(PolygonGroupID))
+	{
+		UE_LOG(LogMeshDescription, Warning, TEXT("CreatePolygonGroupWithID: PolygonGroupID %d already exists."), PolygonGroupID.GetValue());
+		return;
+	}
+
+	MeshDescription.CreatePolygonGroupWithID(PolygonGroupID);
+}
+
+void UMeshDescriptionBase::DeletePolygonGroup(FPolygonGroupID PolygonGroupID)
+{
+	if (!MeshDescription.IsPolygonGroupValid(PolygonGroupID))
+	{
+		UE_LOG(LogMeshDescription, Warning, TEXT("DeletePolygonGroup: FPolygonGroupID %d doesn't exist."), PolygonGroupID.GetValue());
+		return;
+	}
+
+	MeshDescription.DeletePolygonGroup(PolygonGroupID);
+}
+
+bool UMeshDescriptionBase::IsPolygonGroupValid(FPolygonGroupID PolygonGroupID) const
+{
+	return MeshDescription.IsPolygonGroupValid(PolygonGroupID);
+}
+
+
+//////////////////////////////////////////////////////////////////////
+// Vertex operations
+
+bool UMeshDescriptionBase::IsVertexOrphaned(FVertexID VertexID) const
+{
+	if (!MeshDescription.IsVertexValid(VertexID))
+	{
+		UE_LOG(LogMeshDescription, Warning, TEXT("IsVertexOrphaned: VertexID %d doesn't exist."), VertexID.GetValue());
+		return false;
+	}
+
+	return MeshDescription.IsVertexOrphaned(VertexID);
+}
+
+FEdgeID UMeshDescriptionBase::GetVertexPairEdge(FVertexID VertexID0, FVertexID VertexID1) const
+{
+	if (!MeshDescription.IsVertexValid(VertexID0))
+	{
+		UE_LOG(LogMeshDescription, Warning, TEXT("GetVertexPairEdge: VertexID %d doesn't exist."), VertexID0.GetValue());
+		return FEdgeID::Invalid;
+	}
+
+	if (!MeshDescription.IsVertexValid(VertexID1))
+	{
+		UE_LOG(LogMeshDescription, Warning, TEXT("GetVertexPairEdge: VertexID %d doesn't exist."), VertexID1.GetValue());
+		return FEdgeID::Invalid;
+	}
+
+	return MeshDescription.GetVertexPairEdge(VertexID0, VertexID1);
+}
+
+void UMeshDescriptionBase::GetVertexConnectedEdges(FVertexID VertexID, TArray<FEdgeID>& OutEdgeIDs) const
+{
+	if (!MeshDescription.IsVertexValid(VertexID))
+	{
+		UE_LOG(LogMeshDescription, Warning, TEXT("GetVertexConnectedEdges: VertexID %d doesn't exist."), VertexID.GetValue());
+		return;
+	}
+
+	OutEdgeIDs = MeshDescription.GetVertexConnectedEdges(VertexID);
+}
+
+int32 UMeshDescriptionBase::GetNumVertexConnectedEdges(FVertexID VertexID) const
+{
+	if (!MeshDescription.IsVertexValid(VertexID))
+	{
+		UE_LOG(LogMeshDescription, Warning, TEXT("GetNumVertexConnectedEdges: VertexID %d doesn't exist."), VertexID.GetValue());
+		return 0;
+	}
+
+	return MeshDescription.GetNumVertexConnectedEdges(VertexID);
+}
+
+void UMeshDescriptionBase::GetVertexVertexInstances(FVertexID VertexID, TArray<FVertexInstanceID>& OutVertexInstanceIDs) const
+{
+	if (!MeshDescription.IsVertexValid(VertexID))
+	{
+		UE_LOG(LogMeshDescription, Warning, TEXT("GetVertexVertexInstances: VertexID %d doesn't exist."), VertexID.GetValue());
+		return;
+	}
+
+	OutVertexInstanceIDs = MeshDescription.GetVertexVertexInstances(VertexID);
+}
+
+int32 UMeshDescriptionBase::GetNumVertexVertexInstances(FVertexID VertexID) const
+{
+	if (!MeshDescription.IsVertexValid(VertexID))
+	{
+		UE_LOG(LogMeshDescription, Warning, TEXT("GetNumVertexVertexInstances: VertexID %d doesn't exist."), VertexID.GetValue());
+		return 0;
+	}
+
+	return MeshDescription.GetNumVertexVertexInstances(VertexID);
+}
+
+void UMeshDescriptionBase::GetVertexConnectedTriangles(FVertexID VertexID, TArray<FTriangleID>& OutConnectedTriangleIDs) const
+{
+	if (!MeshDescription.IsVertexValid(VertexID))
+	{
+		UE_LOG(LogMeshDescription, Warning, TEXT("GetVertexConnectedTriangles: VertexID %d doesn't exist."), VertexID.GetValue());
+		return;
+	}
+
+	MeshDescription.GetVertexConnectedTriangles(VertexID, OutConnectedTriangleIDs);
+}
+
+int32 UMeshDescriptionBase::GetNumVertexConnectedTriangles(FVertexID VertexID) const
+{
+	if (!MeshDescription.IsVertexValid(VertexID))
+	{
+		UE_LOG(LogMeshDescription, Warning, TEXT("GetNumVertexConnectedTriangles: VertexID %d doesn't exist."), VertexID.GetValue());
+		return 0;
+	}
+
+	return MeshDescription.GetNumVertexConnectedTriangles(VertexID);
+}
+
+void UMeshDescriptionBase::GetVertexConnectedPolygons(FVertexID VertexID, TArray<FPolygonID>& OutConnectedPolygonIDs) const
+{
+	if (!MeshDescription.IsVertexValid(VertexID))
+	{
+		UE_LOG(LogMeshDescription, Warning, TEXT("GetVertexConnectedPolygons: VertexID %d doesn't exist."), VertexID.GetValue());
+		return;
+	}
+
+	MeshDescription.GetVertexConnectedPolygons(VertexID, OutConnectedPolygonIDs);
+}
+
+int32 UMeshDescriptionBase::GetNumVertexConnectedPolygons(FVertexID VertexID) const
+{
+	if (!MeshDescription.IsVertexValid(VertexID))
+	{
+		UE_LOG(LogMeshDescription, Warning, TEXT("GetNumVertexConnectedPolygons: VertexID %d doesn't exist."), VertexID.GetValue());
+		return 0;
+	}
+
+	return MeshDescription.GetNumVertexConnectedPolygons(VertexID);
+}
+
+void UMeshDescriptionBase::GetVertexAdjacentVertices(FVertexID VertexID, TArray<FVertexID>& OutAdjacentVertexIDs) const
+{
+	if (!MeshDescription.IsVertexValid(VertexID))
+	{
+		UE_LOG(LogMeshDescription, Warning, TEXT("GetVertexAdjacentVertices: VertexID %d doesn't exist."), VertexID.GetValue());
+		return;
+	}
+
+	MeshDescription.GetVertexAdjacentVertices(VertexID, OutAdjacentVertexIDs);
+}
+
+FVector UMeshDescriptionBase::GetVertexPosition(FVertexID VertexID) const
+{
+	if (!MeshDescription.IsVertexValid(VertexID))
+	{
+		UE_LOG(LogMeshDescription, Warning, TEXT("GetVertexAttribute: VertexID %d doesn't exist."), VertexID.GetValue());
+		return FVector::ZeroVector;
+	}
+
+	if (!MeshDescription.VertexAttributes().HasAttribute(MeshAttribute::Vertex::Position))
+	{
+		UE_LOG(LogMeshDescription, Warning, TEXT("GetVertexAttribute: VertexAttribute Position doesn't exist."));
+		return FVector::ZeroVector;
+	}
+
+	return MeshDescription.VertexAttributes().GetAttribute<FVector>(VertexID, MeshAttribute::Vertex::Position);
+}
+
+void UMeshDescriptionBase::SetVertexPosition(FVertexID VertexID, const FVector& Position)
+{
+	if (!MeshDescription.IsVertexValid(VertexID))
+	{
+		UE_LOG(LogMeshDescription, Warning, TEXT("SetVertexAttribute: VertexID %d doesn't exist."), VertexID.GetValue());
+		return;
+	}
+
+	if (!MeshDescription.VertexAttributes().HasAttribute(MeshAttribute::Vertex::Position))
+	{
+		UE_LOG(LogMeshDescription, Warning, TEXT("SetVertexAttribute: VertexAttribute Position doesn't exist."));
+		return;
+	}
+
+	MeshDescription.VertexAttributes().SetAttribute(VertexID, MeshAttribute::Vertex::Position, 0, Position);
+}
+
+
+//////////////////////////////////////////////////////////////////////
+// Vertex instance operations
+
+FVertexID UMeshDescriptionBase::GetVertexInstanceVertex(FVertexInstanceID VertexInstanceID) const
+{
+	if (!MeshDescription.IsVertexInstanceValid(VertexInstanceID))
+	{
+		UE_LOG(LogMeshDescription, Warning, TEXT("GetVertexInstanceVertex: VertexInstanceID %d doesn't exist."), VertexInstanceID.GetValue());
+		return FVertexID::Invalid;
+	}
+
+	return MeshDescription.GetVertexInstanceVertex(VertexInstanceID);
+}
+
+FEdgeID UMeshDescriptionBase::GetVertexInstancePairEdge(FVertexInstanceID VertexInstanceID0, FVertexInstanceID VertexInstanceID1) const
+{
+	if (!MeshDescription.IsVertexInstanceValid(VertexInstanceID0))
+	{
+		UE_LOG(LogMeshDescription, Warning, TEXT("GetVertexInstancePairEdge: VertexInstanceID %d doesn't exist."), VertexInstanceID0.GetValue());
+		return FEdgeID::Invalid;
+	}
+
+	if (!MeshDescription.IsVertexInstanceValid(VertexInstanceID1))
+	{
+		UE_LOG(LogMeshDescription, Warning, TEXT("GetVertexInstancePairEdge: VertexInstanceID %d doesn't exist."), VertexInstanceID1.GetValue());
+		return FEdgeID::Invalid;
+	}
+
+	return MeshDescription.GetVertexInstancePairEdge(VertexInstanceID0, VertexInstanceID1);
+}
+
+void UMeshDescriptionBase::GetVertexInstanceConnectedTriangles(FVertexInstanceID VertexInstanceID, TArray<FTriangleID>& OutConnectedTriangleIDs) const
+{
+	if (!MeshDescription.IsVertexInstanceValid(VertexInstanceID))
+	{
+		UE_LOG(LogMeshDescription, Warning, TEXT("GetVertexInstanceConnectedTriangles: VertexInstanceID %d doesn't exist."), VertexInstanceID.GetValue());
+		return;
+	}
+
+	OutConnectedTriangleIDs = MeshDescription.GetVertexInstanceConnectedTriangles(VertexInstanceID);
+}
+
+int32 UMeshDescriptionBase::GetNumVertexInstanceConnectedTriangles(FVertexInstanceID VertexInstanceID) const
+{
+	if (!MeshDescription.IsVertexInstanceValid(VertexInstanceID))
+	{
+		UE_LOG(LogMeshDescription, Warning, TEXT("GetNumVertexInstanceConnectedTriangles: VertexInstanceID %d doesn't exist."), VertexInstanceID.GetValue());
+		return 0;
+	}
+
+	return MeshDescription.GetNumVertexInstanceConnectedTriangles(VertexInstanceID);
+}
+
+void UMeshDescriptionBase::GetVertexInstanceConnectedPolygons(FVertexInstanceID VertexInstanceID, TArray<FPolygonID>& OutConnectedPolygonIDs) const
+{
+	if (!MeshDescription.IsVertexInstanceValid(VertexInstanceID))
+	{
+		UE_LOG(LogMeshDescription, Warning, TEXT("GetVertexInstanceConnectedPolygons: VertexInstanceID %d doesn't exist."), VertexInstanceID.GetValue());
+		return;
+	}
+
+	MeshDescription.GetVertexInstanceConnectedPolygons(VertexInstanceID, OutConnectedPolygonIDs);
+}
+
+int32 UMeshDescriptionBase::GetNumVertexInstanceConnectedPolygons(FVertexInstanceID VertexInstanceID) const
+{
+	if (!MeshDescription.IsVertexInstanceValid(VertexInstanceID))
+	{
+		UE_LOG(LogMeshDescription, Warning, TEXT("GetNumVertexInstanceConnectedPolygons: VertexInstanceID %d doesn't exist."), VertexInstanceID.GetValue());
+		return 0;
+	}
+
+	return MeshDescription.GetNumVertexInstanceConnectedPolygons(VertexInstanceID);
+}
+
+
+//////////////////////////////////////////////////////////////////////
+// Edge operations
+
+bool UMeshDescriptionBase::IsEdgeInternal(FEdgeID EdgeID) const
+{
+	if (!MeshDescription.IsEdgeValid(EdgeID))
+	{
+		UE_LOG(LogMeshDescription, Warning, TEXT("IsEdgeInternal: EdgeID %d doesn't exist."), EdgeID.GetValue());
+		return false;
+	}
+
+	return MeshDescription.IsEdgeInternal(EdgeID);
+}
+
+bool UMeshDescriptionBase::IsEdgeInternalToPolygon(FEdgeID EdgeID, FPolygonID PolygonID) const
+{
+	if (!MeshDescription.IsEdgeValid(EdgeID))
+	{
+		UE_LOG(LogMeshDescription, Warning, TEXT("IsEdgeInternalToPolygon: EdgeID %d doesn't exist."), EdgeID.GetValue());
+		return false;
+	}
+
+	if (!MeshDescription.IsPolygonValid(PolygonID))
+	{
+		UE_LOG(LogMeshDescription, Warning, TEXT("IsEdgeInternalToPolygon: PolygonID %d doesn't exist."), PolygonID.GetValue());
+		return false;
+	}
+
+	return MeshDescription.IsEdgeInternalToPolygon(EdgeID, PolygonID);
+}
+
+void UMeshDescriptionBase::GetEdgeConnectedTriangles(FEdgeID EdgeID, TArray<FTriangleID>& OutConnectedTriangleIDs) const
+{
+	if (!MeshDescription.IsEdgeValid(EdgeID))
+	{
+		UE_LOG(LogMeshDescription, Warning, TEXT("GetEdgeConnectedTriangles: EdgeID %d doesn't exist."), EdgeID.GetValue());
+		return;
+	}
+
+	OutConnectedTriangleIDs = MeshDescription.GetEdgeConnectedTriangles(EdgeID);
+}
+
+int32 UMeshDescriptionBase::GetNumEdgeConnectedTriangles(FEdgeID EdgeID) const
+{
+	if (!MeshDescription.IsEdgeValid(EdgeID))
+	{
+		UE_LOG(LogMeshDescription, Warning, TEXT("GetNumEdgeConnectedTriangles: EdgeID %d doesn't exist."), EdgeID.GetValue());
+		return 0;
+	}
+
+	return MeshDescription.GetNumEdgeConnectedTriangles(EdgeID);
+}
+
+void UMeshDescriptionBase::GetEdgeConnectedPolygons(FEdgeID EdgeID, TArray<FPolygonID>& OutConnectedPolygonIDs) const
+{
+	if (!MeshDescription.IsEdgeValid(EdgeID))
+	{
+		UE_LOG(LogMeshDescription, Warning, TEXT("GetEdgeConnectedPolygons: EdgeID %d doesn't exist."), EdgeID.GetValue());
+		return;
+	}
+
+	MeshDescription.GetEdgeConnectedPolygons(EdgeID, OutConnectedPolygonIDs);
+}
+
+int32 UMeshDescriptionBase::GetNumEdgeConnectedPolygons(FEdgeID EdgeID) const
+{
+	if (!MeshDescription.IsEdgeValid(EdgeID))
+	{
+		UE_LOG(LogMeshDescription, Warning, TEXT("GetNumEdgeConnectedPolygons: EdgeID %d doesn't exist."), EdgeID.GetValue());
+		return 0;
+	}
+
+	return MeshDescription.GetNumEdgeConnectedPolygons(EdgeID);
+}
+
+FVertexID UMeshDescriptionBase::GetEdgeVertex(FEdgeID EdgeID, int32 VertexNumber) const
+{
+	if (!MeshDescription.IsEdgeValid(EdgeID))
+	{
+		UE_LOG(LogMeshDescription, Warning, TEXT("GetEdgeVertex: EdgeID %d doesn't exist."), EdgeID.GetValue());
+		return FVertexID::Invalid;
+	}
+
+	if (VertexNumber != 0 && VertexNumber != 1)
+	{
+		UE_LOG(LogMeshDescription, Warning, TEXT("GetEdgeVertex: invalid vertex number %d."), VertexNumber);
+		return FVertexID::Invalid;
+	}
+
+	return MeshDescription.GetEdgeVertex(EdgeID, VertexNumber);
+}
+
+void UMeshDescriptionBase::GetEdgeVertices(const FEdgeID EdgeID, TArray<FVertexID>& OutVertexIDs) const
+{
+	if (!MeshDescription.IsEdgeValid(EdgeID))
+	{
+		UE_LOG(LogMeshDescription, Warning, TEXT("GetEdgeVertices: EdgeID %d doesn't exist."), EdgeID.GetValue());
+		return;
+	}
+
+	Algo::Copy(MeshDescription.GetEdgeVertices(EdgeID), OutVertexIDs);
+}
+
+
+//////////////////////////////////////////////////////////////////////
+// Triangle operations
+
+FPolygonID UMeshDescriptionBase::GetTrianglePolygon(FTriangleID TriangleID) const
+{
+	if (!MeshDescription.IsTriangleValid(TriangleID))
+	{
+		UE_LOG(LogMeshDescription, Warning, TEXT("GetTrianglePolygon: TriangleID %d doesn't exist."), TriangleID.GetValue());
+		return FPolygonID::Invalid;
+	}
+
+	return MeshDescription.GetTrianglePolygon(TriangleID);
+}
+
+FPolygonGroupID UMeshDescriptionBase::GetTrianglePolygonGroup(FTriangleID TriangleID) const
+{
+	if (!MeshDescription.IsTriangleValid(TriangleID))
+	{
+		UE_LOG(LogMeshDescription, Warning, TEXT("GetTrianglePolygonGroup: TriangleID %d doesn't exist."), TriangleID.GetValue());
+		return FPolygonGroupID::Invalid;
+	}
+
+	return MeshDescription.GetTrianglePolygonGroup(TriangleID);
+}
+
+bool UMeshDescriptionBase::IsTrianglePartOfNgon(FTriangleID TriangleID) const
+{
+	if (!MeshDescription.IsTriangleValid(TriangleID))
+	{
+		UE_LOG(LogMeshDescription, Warning, TEXT("IsTrianglePartOfNgon: TriangleID %d doesn't exist."), TriangleID.GetValue());
+		return false;
+	}
+
+	return MeshDescription.IsTrianglePartOfNgon(TriangleID);
+}
+
+void UMeshDescriptionBase::GetTriangleVertexInstances(FTriangleID TriangleID, TArray<FVertexInstanceID>& OutVertexInstanceIDs) const
+{
+	if (!MeshDescription.IsTriangleValid(TriangleID))
+	{
+		UE_LOG(LogMeshDescription, Warning, TEXT("GetTriangleVertexInstances: TriangleID %d doesn't exist."), TriangleID.GetValue());
+		return;
+	}
+
+	Algo::Copy(MeshDescription.GetTriangleVertexInstances(TriangleID), OutVertexInstanceIDs);
+}
+
+FVertexInstanceID UMeshDescriptionBase::GetTriangleVertexInstance(FTriangleID TriangleID, int32 Index) const
+{
+	if (!MeshDescription.IsTriangleValid(TriangleID))
+	{
+		UE_LOG(LogMeshDescription, Warning, TEXT("GetTriangleVertexInstance: TriangleID %d doesn't exist."), TriangleID.GetValue());
+		return FVertexInstanceID::Invalid;
+	}
+
+	if (Index < 0 || Index > 2)
+	{
+		UE_LOG(LogMeshDescription, Warning, TEXT("GetTriangleVertexInstance: invalid vertex index %d."), Index);
+		return FVertexInstanceID::Invalid;
+	}
+
+	return MeshDescription.GetTriangleVertexInstance(TriangleID, Index);
+}
+
+void UMeshDescriptionBase::GetTriangleVertices(FTriangleID TriangleID, TArray<FVertexID>& OutVertexIDs) const
+{
+	if (!MeshDescription.IsTriangleValid(TriangleID))
+	{
+		UE_LOG(LogMeshDescription, Warning, TEXT("GetTriangleVertices: TriangleID %d doesn't exist."), TriangleID.GetValue());
+		return;
+	}
+
+	OutVertexIDs.SetNumUninitialized(3);
+	MeshDescription.GetTriangleVertices(TriangleID, OutVertexIDs);
+}
+
+void UMeshDescriptionBase::GetTriangleEdges(FTriangleID TriangleID, TArray<FEdgeID>& OutEdgeIDs) const
+{
+	if (!MeshDescription.IsTriangleValid(TriangleID))
+	{
+		UE_LOG(LogMeshDescription, Warning, TEXT("GetTriangleEdges: TriangleID %d doesn't exist."), TriangleID.GetValue());
+		return;
+	}
+
+	OutEdgeIDs.SetNumUninitialized(3);
+	MeshDescription.GetTriangleEdges(TriangleID, OutEdgeIDs);
+}
+
+void UMeshDescriptionBase::GetTriangleAdjacentTriangles(FTriangleID TriangleID, TArray<FTriangleID>& OutTriangleIDs) const
+{
+	if (!MeshDescription.IsTriangleValid(TriangleID))
+	{
+		UE_LOG(LogMeshDescription, Warning, TEXT("GetTriangleAdjacentTriangles: TriangleID %d doesn't exist."), TriangleID.GetValue());
+		return;
+	}
+
+	MeshDescription.GetTriangleAdjacentTriangles(TriangleID, OutTriangleIDs);
+}
+
+FVertexInstanceID UMeshDescriptionBase::GetVertexInstanceForTriangleVertex(FTriangleID TriangleID, FVertexID VertexID) const
+{
+	if (!MeshDescription.IsTriangleValid(TriangleID))
+	{
+		UE_LOG(LogMeshDescription, Warning, TEXT("GetVertexInstanceForTriangleVertex: TriangleID %d doesn't exist."), TriangleID.GetValue());
+		return FVertexInstanceID::Invalid;
+	}
+
+	return MeshDescription.GetVertexInstanceForTriangleVertex(TriangleID, VertexID);
+}
+
+
+//////////////////////////////////////////////////////////////////////
+// Polygon operations
+
+void UMeshDescriptionBase::GetPolygonTriangles(FPolygonID PolygonID, TArray<FTriangleID>& OutTriangleIDs) const
+{
+	if (!MeshDescription.IsPolygonValid(PolygonID))
+	{
+		UE_LOG(LogMeshDescription, Warning, TEXT("GetPolygonTriangles: PolygonID %d doesn't exist."), PolygonID.GetValue());
+		return;
+	}
+
+	OutTriangleIDs = MeshDescription.GetPolygonTriangleIDs(PolygonID);
+}
+
+int32 UMeshDescriptionBase::GetNumPolygonTriangles(FPolygonID PolygonID) const
+{
+	if (!MeshDescription.IsPolygonValid(PolygonID))
+	{
+		UE_LOG(LogMeshDescription, Warning, TEXT("GetNumPolygonTriangles: PolygonID %d doesn't exist."), PolygonID.GetValue());
+		return 0;
+	}
+
+	return MeshDescription.GetNumPolygonTriangles(PolygonID);
+}
+
+void UMeshDescriptionBase::GetPolygonVertexInstances(FPolygonID PolygonID, TArray<FVertexInstanceID>& OutVertexInstanceIDs) const
+{
+	if (!MeshDescription.IsPolygonValid(PolygonID))
+	{
+		UE_LOG(LogMeshDescription, Warning, TEXT("GetPolygonVertexInstances: PolygonID %d doesn't exist."), PolygonID.GetValue());
+		return;
+	}
+
+	OutVertexInstanceIDs = MeshDescription.GetPolygonVertexInstances(PolygonID);
+}
+
+int32 UMeshDescriptionBase::GetNumPolygonVertices(FPolygonID PolygonID) const
+{
+	if (!MeshDescription.IsPolygonValid(PolygonID))
+	{
+		UE_LOG(LogMeshDescription, Warning, TEXT("GetNumPolygonVertices: PolygonID %d doesn't exist."), PolygonID.GetValue());
+		return 0;
+	}
+
+	return MeshDescription.GetNumPolygonVertices(PolygonID);
+}
+
+void UMeshDescriptionBase::GetPolygonVertices(FPolygonID PolygonID, TArray<FVertexID>& OutVertexIDs) const
+{
+	if (!MeshDescription.IsPolygonValid(PolygonID))
+	{
+		UE_LOG(LogMeshDescription, Warning, TEXT("GetPolygonVertices: PolygonID %d doesn't exist."), PolygonID.GetValue());
+		return;
+	}
+
+	MeshDescription.GetPolygonVertices(PolygonID, OutVertexIDs);
+}
+
+void UMeshDescriptionBase::GetPolygonPerimeterEdges(FPolygonID PolygonID, TArray<FEdgeID>& OutEdgeIDs) const
+{
+	if (!MeshDescription.IsPolygonValid(PolygonID))
+	{
+		UE_LOG(LogMeshDescription, Warning, TEXT("GetPolygonPerimeterEdges: PolygonID %d doesn't exist."), PolygonID.GetValue());
+		return;
+	}
+
+	MeshDescription.GetPolygonPerimeterEdges(PolygonID, OutEdgeIDs);
+}
+
+void UMeshDescriptionBase::GetPolygonInternalEdges(FPolygonID PolygonID, TArray<FEdgeID>& OutEdgeIDs) const
+{
+	if (!MeshDescription.IsPolygonValid(PolygonID))
+	{
+		UE_LOG(LogMeshDescription, Warning, TEXT("GetPolygonInternalEdges: PolygonID %d doesn't exist."), PolygonID.GetValue());
+		return;
+	}
+
+	MeshDescription.GetPolygonInternalEdges(PolygonID, OutEdgeIDs);
+}
+
+int32 UMeshDescriptionBase::GetNumPolygonInternalEdges(FPolygonID PolygonID) const
+{
+	if (!MeshDescription.IsPolygonValid(PolygonID))
+	{
+		UE_LOG(LogMeshDescription, Warning, TEXT("GetNumPolygonInternalEdges: PolygonID %d doesn't exist."), PolygonID.GetValue());
+		return 0;
+	}
+
+	return MeshDescription.GetNumPolygonInternalEdges(PolygonID);
+}
+
+void UMeshDescriptionBase::GetPolygonAdjacentPolygons(FPolygonID PolygonID, TArray<FPolygonID>& OutPolygonIDs) const
+{
+	if (!MeshDescription.IsPolygonValid(PolygonID))
+	{
+		UE_LOG(LogMeshDescription, Warning, TEXT("GetPolygonAdjacentPolygons: PolygonID %d doesn't exist."), PolygonID.GetValue());
+		return;
+	}
+
+	MeshDescription.GetPolygonAdjacentPolygons(PolygonID, OutPolygonIDs);
+}
+
+FPolygonGroupID UMeshDescriptionBase::GetPolygonPolygonGroup(FPolygonID PolygonID) const
+{
+	if (!MeshDescription.IsPolygonValid(PolygonID))
+	{
+		UE_LOG(LogMeshDescription, Warning, TEXT("GetPolygonPolygonGroup: PolygonID %d doesn't exist."), PolygonID.GetValue());
+		return FPolygonGroupID::Invalid;
+	}
+
+	return MeshDescription.GetPolygonPolygonGroup(PolygonID);
+}
+
+FVertexInstanceID UMeshDescriptionBase::GetVertexInstanceForPolygonVertex(FPolygonID PolygonID, FVertexID VertexID) const
+{
+	if (!MeshDescription.IsPolygonValid(PolygonID))
+	{
+		UE_LOG(LogMeshDescription, Warning, TEXT("GetVertexInstanceForPolygonVertex: PolygonID %d doesn't exist."), PolygonID.GetValue());
+		return FVertexInstanceID::Invalid;
+	}
+
+	return MeshDescription.GetVertexInstanceForPolygonVertex(PolygonID, VertexID);
+}
+
+void UMeshDescriptionBase::SetPolygonVertexInstance(FPolygonID PolygonID, int32 PerimeterIndex, FVertexInstanceID VertexInstanceID)
+{
+	if (!MeshDescription.IsPolygonValid(PolygonID))
+	{
+		UE_LOG(LogMeshDescription, Warning, TEXT("SetPolygonVertexInstance: PolygonID %d doesn't exist."), PolygonID.GetValue());
+		return;
+	}
+
+	if (PerimeterIndex >= MeshDescription.GetNumPolygonVertices(PolygonID))
+	{
+		UE_LOG(LogMeshDescription, Warning, TEXT("SetPolygonVertexInstance: Out of range vertex index %d."), PerimeterIndex);
+		return;
+	}
+
+	if (!MeshDescription.IsVertexInstanceValid(VertexInstanceID))
+	{
+		UE_LOG(LogMeshDescription, Warning, TEXT("SetPolygonVertexInstance: VertexInstanceID %d doesn't exist."), VertexInstanceID.GetValue());
+		return;
+	}
+
+	MeshDescription.SetPolygonVertexInstance(PolygonID, PerimeterIndex, VertexInstanceID);
+}
+
+void UMeshDescriptionBase::SetPolygonPolygonGroup(FPolygonID PolygonID, FPolygonGroupID PolygonGroupID)
+{
+	if (!MeshDescription.IsPolygonValid(PolygonID))
+	{
+		UE_LOG(LogMeshDescription, Warning, TEXT("SetPolygonPolygonGroup: PolygonID %d doesn't exist."), PolygonID.GetValue());
+		return;
+	}
+
+	if (!MeshDescription.IsPolygonGroupValid(PolygonGroupID))
+	{
+		UE_LOG(LogMeshDescription, Warning, TEXT("SetPolygonPolygonGroup: PolygonGroupID %d doesn't exist."), PolygonGroupID.GetValue());
+		return;
+	}
+
+	MeshDescription.SetPolygonPolygonGroup(PolygonID, PolygonGroupID);
+}
+
+void UMeshDescriptionBase::ReversePolygonFacing(FPolygonID PolygonID)
+{
+	if (!MeshDescription.IsPolygonValid(PolygonID))
+	{
+		UE_LOG(LogMeshDescription, Warning, TEXT("ReversePolygonFacing: PolygonID %d doesn't exist."), PolygonID.GetValue());
+		return;
+	}
+
+	MeshDescription.ReversePolygonFacing(PolygonID);
+}
+
+void UMeshDescriptionBase::ComputePolygonTriangulation(FPolygonID PolygonID)
+{
+	if (!MeshDescription.IsPolygonValid(PolygonID))
+	{
+		UE_LOG(LogMeshDescription, Warning, TEXT("ComputePolygonTriangulation: PolygonID %d doesn't exist."), PolygonID.GetValue());
+		return;
+	}
+
+	MeshDescription.ComputePolygonTriangulation(PolygonID);
+}
+
+
+//////////////////////////////////////////////////////////////////////
+// Polygon group operations
+
+void UMeshDescriptionBase::GetPolygonGroupPolygons(FPolygonGroupID PolygonGroupID, TArray<FPolygonID>& OutPolygonIDs) const
+{
+	if (!MeshDescription.IsPolygonGroupValid(PolygonGroupID))
+	{
+		UE_LOG(LogMeshDescription, Warning, TEXT("GetPolygonGroupPolygons: PolygonGroupID %d doesn't exist."), PolygonGroupID.GetValue());
+		return;
+	}
+
+	OutPolygonIDs = MeshDescription.GetPolygonGroupPolygons(PolygonGroupID);
+}
+
+int32 UMeshDescriptionBase::GetNumPolygonGroupPolygons(FPolygonGroupID PolygonGroupID) const
+{
+	if (!MeshDescription.IsPolygonGroupValid(PolygonGroupID))
+	{
+		UE_LOG(LogMeshDescription, Warning, TEXT("GetNumPolygonGroupPolygons: PolygonGroupID %d doesn't exist."), PolygonGroupID.GetValue());
+		return 0;
+	}
+
+	return MeshDescription.GetNumPolygonGroupPolygons(PolygonGroupID);
+}
diff --git a/Engine/Source/Runtime/MeshDescription/Public/MeshAttributeArray.h b/Engine/Source/Runtime/MeshDescription/Public/MeshAttributeArray.h
--- a/Engine/Source/Runtime/MeshDescription/Public/MeshAttributeArray.h
+++ b/Engine/Source/Runtime/MeshDescription/Public/MeshAttributeArray.h
@@ -5,6 +5,7 @@
 #include "CoreMinimal.h"
 #include "MeshTypes.h"
 #include "MeshElementRemappings.h"
+#include "Misc/TVariant.h"
 #include "UObject/ReleaseObjectVersion.h"
 
 
@@ -56,6 +57,13 @@ template <typename T, typename... Ts> struct TTupleElement<0, TTuple<T, Ts...>>
 template <typename... Ts> struct TTupleElement<0, TTuple<Ts...>> { static_assert( sizeof...( Ts ) > 0, "TTuple element index out of range" ); };
 
 
+/**
+ * Helper template which generates a TVariant of all supported attribute types.
+ */
+template <typename Tuple> struct TVariantFromTuple;
+template <typename... Ts> struct TVariantFromTuple<TTuple<Ts...>> { using Type = TVariant<FEmptyVariantState, Ts...>; };
+
+
 /**
  * Class which implements a function jump table to be automatically generated at compile time.
  * This is used by TAttributesSet to provide O(1) dispatch by attribute type at runtime.
diff --git a/Engine/Source/Runtime/MeshDescription/Public/MeshAttributes.h b/Engine/Source/Runtime/MeshDescription/Public/MeshAttributes.h
--- a/Engine/Source/Runtime/MeshDescription/Public/MeshAttributes.h
+++ b/Engine/Source/Runtime/MeshDescription/Public/MeshAttributes.h
@@ -3,46 +3,49 @@
 #pragma once
 
 #include "CoreMinimal.h"
+#include "MeshDescription.h"
 
-// @todo mesh description: for now, these attributes are hardcoded in the MeshDescription module.
-// It would be better if these were defined in separate modules, one per mesh type which uses MeshDescription.
-// This would give an extensible and modular structure whereby a mesh type defines its valid attributes and adapters.
 namespace MeshAttribute
 {
 	namespace Vertex
 	{
 		extern MESHDESCRIPTION_API const FName Position;
-		extern MESHDESCRIPTION_API const FName CornerSharpness;
 	}
+}
 
-	namespace VertexInstance
-	{
-		extern MESHDESCRIPTION_API const FName TextureCoordinate;
-		extern MESHDESCRIPTION_API const FName Normal;
-		extern MESHDESCRIPTION_API const FName Tangent;
-		extern MESHDESCRIPTION_API const FName BinormalSign;
-		extern MESHDESCRIPTION_API const FName Color;
-	}
 
-	namespace Edge
-	{
-		extern MESHDESCRIPTION_API const FName IsHard;
-		extern MESHDESCRIPTION_API const FName IsUVSeam;
-		extern MESHDESCRIPTION_API const FName CreaseSharpness;
-	}
+class MESHDESCRIPTION_API FMeshAttributes
+{
+public:
+	explicit FMeshAttributes(FMeshDescription& InMeshDescription)
+		: MeshDescription(InMeshDescription)
+	{}
 
-	namespace Polygon
-	{
-		extern MESHDESCRIPTION_API const FName Normal;
-		extern MESHDESCRIPTION_API const FName Tangent;
-		extern MESHDESCRIPTION_API const FName Binormal;
-		extern MESHDESCRIPTION_API const FName Center;
-	}
+	virtual ~FMeshAttributes() = default;
 
-	namespace PolygonGroup
-	{
-		extern MESHDESCRIPTION_API const FName ImportedMaterialSlotName;
-		extern MESHDESCRIPTION_API const FName EnableCollision;
-		extern MESHDESCRIPTION_API const FName CastShadow;
-	}
-}
+	virtual void Register();
+
+	/** Accessors for cached vertex position array */
+	TVertexAttributesRef<FVector> GetVertexPositions() { return MeshDescription.VertexAttributes().GetAttributesRef<FVector>(MeshAttribute::Vertex::Position); }
+	TVertexAttributesConstRef<FVector> GetVertexPositions() const { return MeshDescription.VertexAttributes().GetAttributesRef<FVector>(MeshAttribute::Vertex::Position); }
+
+protected:
+
+	FMeshDescription& MeshDescription;
+};
+
+
+class FMeshConstAttributes
+{
+public:
+	explicit FMeshConstAttributes(const FMeshDescription& InMeshDescription)
+		: MeshDescription(InMeshDescription)
+	{}
+
+	/** Accessors for cached vertex position array */
+	TVertexAttributesConstRef<FVector> GetVertexPositions() const { return MeshDescription.VertexAttributes().GetAttributesRef<FVector>(MeshAttribute::Vertex::Position); }
+
+protected:
+
+	const FMeshDescription& MeshDescription;
+};
