Copying //UE4/Dev-Editor to Dev-Main (//UE4/Dev-Main) for 4.24

diff --git a/Engine/Source/Runtime/MeshDescription/Private/MeshDescription.cpp b/Engine/Source/Runtime/MeshDescription/Private/MeshDescription.cpp
--- a/Engine/Source/Runtime/MeshDescription/Private/MeshDescription.cpp
+++ b/Engine/Source/Runtime/MeshDescription/Private/MeshDescription.cpp
@@ -23,6 +23,13 @@ void UDEPRECATED_MeshDescription::Serialize( FArchive& Ar )
 
 PRAGMA_DISABLE_DEPRECATION_WARNINGS
 
+FMeshDescription::FMeshDescription()
+{
+	// Minimal requirement is that vertices have a Position attribute
+	VertexAttributesSet.RegisterAttribute(MeshAttribute::Vertex::Position, 1, FVector::ZeroVector, EMeshAttributeFlags::Lerpable);
+}
+
+
 void FMeshDescription::Serialize(FArchive& Ar)
 {
 	Ar.UsingCustomVersion(FReleaseObjectVersion::GUID);
@@ -1229,6 +1236,39 @@ void FMeshDescription::ComputePolygonTriangulation(const FPolygonID PolygonID, T
 }
 
 
+FBoxSphereBounds FMeshDescription::GetBounds() const
+{
+	FBoxSphereBounds BoundingBoxAndSphere;
+
+	FBox BoundingBox;
+	BoundingBox.Init();
+
+	TVertexAttributesConstRef<FVector> VertexPositions = VertexAttributes().GetAttributesRef<FVector>(MeshAttribute::Vertex::Position);
+
+	for (const FVertexID VertexID : Vertices().GetElementIDs())
+	{
+		if (!IsVertexOrphaned(VertexID))
+		{
+			BoundingBox += VertexPositions[VertexID];
+		}
+	}
+
+	BoundingBox.GetCenterAndExtents(BoundingBoxAndSphere.Origin, BoundingBoxAndSphere.BoxExtent);
+
+	// Calculate the bounding sphere, using the center of the bounding box as the origin.
+	BoundingBoxAndSphere.SphereRadius = 0.0f;
+
+	for (const FVertexID VertexID : Vertices().GetElementIDs())
+	{
+		if (!IsVertexOrphaned(VertexID))
+		{
+			BoundingBoxAndSphere.SphereRadius = FMath::Max((VertexPositions[VertexID] - BoundingBoxAndSphere.Origin).Size(), BoundingBoxAndSphere.SphereRadius);
+		}
+	}
+
+	return BoundingBoxAndSphere;
+}
+
 void FMeshDescription::TriangulateMesh()
 {
 	// Perform triangulation directly into mesh polygons
@@ -1239,6 +1279,47 @@ void FMeshDescription::TriangulateMesh()
 }
 
 
+namespace MeshAttribute_
+{
+	namespace Vertex
+	{
+		const FName CornerSharpness("CornerSharpness");
+	}
+
+	namespace VertexInstance
+	{
+		const FName TextureCoordinate("TextureCoordinate");
+		const FName Normal("Normal");
+		const FName Tangent("Tangent");
+		const FName BinormalSign("BinormalSign");
+		const FName Color("Color");
+	}
+
+	namespace Edge
+	{
+		const FName IsHard("IsHard");
+		const FName IsUVSeam("IsUVSeam");
+		const FName CreaseSharpness("CreaseSharpness");
+	}
+
+	namespace Polygon
+	{
+		const FName Normal("Normal");
+		const FName Tangent("Tangent");
+		const FName Binormal("Binormal");
+		const FName Center("Center");
+	}
+
+	namespace PolygonGroup
+	{
+		const FName ImportedMaterialSlotName("ImportedMaterialSlotName");
+		const FName EnableCollision("EnableCollision");
+		const FName CastShadow("CastShadow");
+	}
+}
+
+
+
 bool FMeshDescription::ComputePolygonTangentsAndNormals(const FPolygonID PolygonID
 	, float ComparisonThreshold
 	, const TVertexAttributesRef<FVector> VertexPositions
@@ -1308,11 +1389,11 @@ bool FMeshDescription::ComputePolygonTangentsAndNormals(const FPolygonID Polygon
 void FMeshDescription::ComputePolygonTangentsAndNormals(const TArray<FPolygonID>& PolygonIDs, float ComparisonThreshold)
 {
 	const TVertexAttributesRef<FVector> VertexPositions = VertexAttributes().GetAttributesRef<FVector>(MeshAttribute::Vertex::Position);
-	const TVertexInstanceAttributesRef<FVector2D> VertexUVs = VertexInstanceAttributes().GetAttributesRef<FVector2D>(MeshAttribute::VertexInstance::TextureCoordinate);
-	TPolygonAttributesRef<FVector> PolygonNormals = PolygonAttributes().GetAttributesRef<FVector>(MeshAttribute::Polygon::Normal);
-	TPolygonAttributesRef<FVector> PolygonTangents = PolygonAttributes().GetAttributesRef<FVector>(MeshAttribute::Polygon::Tangent);
-	TPolygonAttributesRef<FVector> PolygonBinormals = PolygonAttributes().GetAttributesRef<FVector>(MeshAttribute::Polygon::Binormal);
-	TPolygonAttributesRef<FVector> PolygonCenters = PolygonAttributes().GetAttributesRef<FVector>(MeshAttribute::Polygon::Center);
+	const TVertexInstanceAttributesRef<FVector2D> VertexUVs = VertexInstanceAttributes().GetAttributesRef<FVector2D>(MeshAttribute_::VertexInstance::TextureCoordinate);
+	TPolygonAttributesRef<FVector> PolygonNormals = PolygonAttributes().GetAttributesRef<FVector>(MeshAttribute_::Polygon::Normal);
+	TPolygonAttributesRef<FVector> PolygonTangents = PolygonAttributes().GetAttributesRef<FVector>(MeshAttribute_::Polygon::Tangent);
+	TPolygonAttributesRef<FVector> PolygonBinormals = PolygonAttributes().GetAttributesRef<FVector>(MeshAttribute_::Polygon::Binormal);
+	TPolygonAttributesRef<FVector> PolygonCenters = PolygonAttributes().GetAttributesRef<FVector>(MeshAttribute_::Polygon::Center);
 
 	TArray<FPolygonID> DegeneratePolygonIDs;
 	for (const FPolygonID PolygonID : PolygonIDs)
@@ -1374,7 +1455,7 @@ void FMeshDescription::GetConnectedSoftEdges(const FVertexID VertexID, TArray<FE
 {
 	OutConnectedSoftEdges.Reset();
 
-	TEdgeAttributesConstRef<bool> EdgeHardnesses = EdgeAttributes().GetAttributesRef<bool>(MeshAttribute::Edge::IsHard);
+	TEdgeAttributesConstRef<bool> EdgeHardnesses = EdgeAttributes().GetAttributesRef<bool>(MeshAttribute_::Edge::IsHard);
 	for (const FEdgeID ConnectedEdgeID : VertexArray[VertexID].ConnectedEdgeIDs)
 	{
 		if (!EdgeHardnesses[ConnectedEdgeID])
@@ -1608,13 +1689,13 @@ void FMeshDescription::ComputeTangentsAndNormals(const FVertexInstanceID VertexI
 
 void FMeshDescription::ComputeTangentsAndNormals(const TArray<FVertexInstanceID>& VertexInstanceIDs, EComputeNTBsOptions ComputeNTBsOptions)
 {
-	const TPolygonAttributesRef<FVector> PolygonNormals = PolygonAttributes().GetAttributesRef<FVector>(MeshAttribute::Polygon::Normal);
-	const TPolygonAttributesRef<FVector> PolygonTangents = PolygonAttributes().GetAttributesRef<FVector>(MeshAttribute::Polygon::Tangent);
-	const TPolygonAttributesRef<FVector> PolygonBinormals = PolygonAttributes().GetAttributesRef<FVector>(MeshAttribute::Polygon::Binormal);
+	const TPolygonAttributesRef<FVector> PolygonNormals = PolygonAttributes().GetAttributesRef<FVector>(MeshAttribute_::Polygon::Normal);
+	const TPolygonAttributesRef<FVector> PolygonTangents = PolygonAttributes().GetAttributesRef<FVector>(MeshAttribute_::Polygon::Tangent);
+	const TPolygonAttributesRef<FVector> PolygonBinormals = PolygonAttributes().GetAttributesRef<FVector>(MeshAttribute_::Polygon::Binormal);
 
-	TVertexInstanceAttributesRef<FVector> VertexNormals = VertexInstanceAttributes().GetAttributesRef<FVector>(MeshAttribute::VertexInstance::Normal);
-	TVertexInstanceAttributesRef<FVector> VertexTangents = VertexInstanceAttributes().GetAttributesRef<FVector>(MeshAttribute::VertexInstance::Tangent);
-	TVertexInstanceAttributesRef<float> VertexBinormalSigns = VertexInstanceAttributes().GetAttributesRef<float>(MeshAttribute::VertexInstance::BinormalSign);
+	TVertexInstanceAttributesRef<FVector> VertexNormals = VertexInstanceAttributes().GetAttributesRef<FVector>(MeshAttribute_::VertexInstance::Normal);
+	TVertexInstanceAttributesRef<FVector> VertexTangents = VertexInstanceAttributes().GetAttributesRef<FVector>(MeshAttribute_::VertexInstance::Tangent);
+	TVertexInstanceAttributesRef<float> VertexBinormalSigns = VertexInstanceAttributes().GetAttributesRef<float>(MeshAttribute_::VertexInstance::BinormalSign);
 
 	for (const FVertexInstanceID VertexInstanceID : VertexInstanceIDs)
 	{
@@ -1625,13 +1706,13 @@ void FMeshDescription::ComputeTangentsAndNormals(const TArray<FVertexInstanceID>
 
 void FMeshDescription::ComputeTangentsAndNormals(EComputeNTBsOptions ComputeNTBsOptions)
 {
-	const TPolygonAttributesRef<FVector> PolygonNormals = PolygonAttributes().GetAttributesRef<FVector>(MeshAttribute::Polygon::Normal);
-	const TPolygonAttributesRef<FVector> PolygonTangents = PolygonAttributes().GetAttributesRef<FVector>(MeshAttribute::Polygon::Tangent);
-	const TPolygonAttributesRef<FVector> PolygonBinormals = PolygonAttributes().GetAttributesRef<FVector>(MeshAttribute::Polygon::Binormal);
+	const TPolygonAttributesRef<FVector> PolygonNormals = PolygonAttributes().GetAttributesRef<FVector>(MeshAttribute_::Polygon::Normal);
+	const TPolygonAttributesRef<FVector> PolygonTangents = PolygonAttributes().GetAttributesRef<FVector>(MeshAttribute_::Polygon::Tangent);
+	const TPolygonAttributesRef<FVector> PolygonBinormals = PolygonAttributes().GetAttributesRef<FVector>(MeshAttribute_::Polygon::Binormal);
 
-	TVertexInstanceAttributesRef<FVector> VertexNormals = VertexInstanceAttributes().GetAttributesRef<FVector>(MeshAttribute::VertexInstance::Normal);
-	TVertexInstanceAttributesRef<FVector> VertexTangents = VertexInstanceAttributes().GetAttributesRef<FVector>(MeshAttribute::VertexInstance::Tangent);
-	TVertexInstanceAttributesRef<float> VertexBinormalSigns = VertexInstanceAttributes().GetAttributesRef<float>(MeshAttribute::VertexInstance::BinormalSign);
+	TVertexInstanceAttributesRef<FVector> VertexNormals = VertexInstanceAttributes().GetAttributesRef<FVector>(MeshAttribute_::VertexInstance::Normal);
+	TVertexInstanceAttributesRef<FVector> VertexTangents = VertexInstanceAttributes().GetAttributesRef<FVector>(MeshAttribute_::VertexInstance::Tangent);
+	TVertexInstanceAttributesRef<float> VertexBinormalSigns = VertexInstanceAttributes().GetAttributesRef<float>(MeshAttribute_::VertexInstance::BinormalSign);
 
 	for (const FVertexInstanceID VertexInstanceID : VertexInstances().GetElementIDs())
 	{
@@ -1642,8 +1723,8 @@ void FMeshDescription::ComputeTangentsAndNormals(EComputeNTBsOptions ComputeNTBs
 
 void FMeshDescription::DetermineEdgeHardnessesFromVertexInstanceNormals( const float Tolerance )
 {
-	const TVertexInstanceAttributesRef<FVector> VertexNormals = VertexInstanceAttributes().GetAttributesRef<FVector>( MeshAttribute::VertexInstance::Normal );
-	TEdgeAttributesRef<bool> EdgeHardnesses = EdgeAttributes().GetAttributesRef<bool>( MeshAttribute::Edge::IsHard );
+	const TVertexInstanceAttributesRef<FVector> VertexNormals = VertexInstanceAttributes().GetAttributesRef<FVector>( MeshAttribute_::VertexInstance::Normal );
+	TEdgeAttributesRef<bool> EdgeHardnesses = EdgeAttributes().GetAttributesRef<bool>( MeshAttribute_::Edge::IsHard );
 
 	// Holds unique vertex instance IDs for a given edge vertex
 	// @todo: use TMemStackAllocator or similar to avoid expensive allocations
@@ -1706,8 +1787,8 @@ void FMeshDescription::DetermineEdgeHardnessesFromVertexInstanceNormals( const f
 
 void FMeshDescription::DetermineUVSeamsFromUVs( const int32 UVIndex, const float Tolerance )
 {
-	const TVertexInstanceAttributesRef<FVector2D> VertexUVs = VertexInstanceAttributes().GetAttributesRef<FVector2D>( MeshAttribute::VertexInstance::TextureCoordinate );
-	TEdgeAttributesRef<bool> EdgeUVSeams = EdgeAttributes().GetAttributesRef<bool>( MeshAttribute::Edge::IsUVSeam );
+	const TVertexInstanceAttributesRef<FVector2D> VertexUVs = VertexInstanceAttributes().GetAttributesRef<FVector2D>( MeshAttribute_::VertexInstance::TextureCoordinate );
+	TEdgeAttributesRef<bool> EdgeUVSeams = EdgeAttributes().GetAttributesRef<bool>( MeshAttribute_::Edge::IsUVSeam );
 
 	// Holds unique vertex instance IDs for a given edge vertex
 	// @todo: use TMemStackAllocator or similar to avoid expensive allocations
@@ -1769,7 +1850,7 @@ void FMeshDescription::DetermineUVSeamsFromUVs( const int32 UVIndex, const float
 
 void FMeshDescription::GetPolygonsInSameChartAsPolygon( const FPolygonID PolygonID, TArray<FPolygonID>& OutPolygonIDs )
 {
-	const TEdgeAttributesRef<bool> EdgeUVSeams = EdgeAttributes().GetAttributesRef<bool>( MeshAttribute::Edge::IsUVSeam );
+	const TEdgeAttributesRef<bool> EdgeUVSeams = EdgeAttributes().GetAttributesRef<bool>( MeshAttribute_::Edge::IsUVSeam );
 	const int32 NumPolygons = Polygons().Num();
 
 	// This holds the results - all polygon IDs which are in the same UV chart
@@ -1871,7 +1952,7 @@ void FMeshDescription::ReverseAllPolygonFacing()
 
 void FMeshDescription::RemapPolygonGroups(const TMap<FPolygonGroupID, FPolygonGroupID>& Remap)
 {
-	TPolygonGroupAttributesRef<FName> PolygonGroupNames = PolygonGroupAttributes().GetAttributesRef<FName>(MeshAttribute::PolygonGroup::ImportedMaterialSlotName);
+	TPolygonGroupAttributesRef<FName> PolygonGroupNames = PolygonGroupAttributes().GetAttributesRef<FName>(MeshAttribute_::PolygonGroup::ImportedMaterialSlotName);
 
 	struct FOldPolygonGroupData
 	{
diff --git a/Engine/Source/Runtime/MeshDescription/Private/MeshDescriptionBase.cpp b/Engine/Source/Runtime/MeshDescription/Private/MeshDescriptionBase.cpp
new file mode 100644
--- /dev/null
+++ b/Engine/Source/Runtime/MeshDescription/Private/MeshDescriptionBase.cpp
@@ -0,0 +1,992 @@
+// Copyright 1998-2019 Epic Games, Inc. All Rights Reserved.
+
+#include "MeshDescriptionBase.h"
+#include "Algo/Copy.h"
+
+
+
+void UMeshDescriptionBase::Serialize(FArchive& Ar)
+{
+	Super::Serialize(Ar);
+}
+
+void UMeshDescriptionBase::RegisterAttributes()
+{
+	RequiredAttributes = MakeUnique<FMeshAttributes>(MeshDescription);
+	RequiredAttributes->Register();
+}
+
+void UMeshDescriptionBase::Reset()
+{
+	MeshDescription = FMeshDescription();
+	RegisterAttributes();
+}
+
+void UMeshDescriptionBase::Empty()
+{
+	MeshDescription.Empty();
+}
+
+bool UMeshDescriptionBase::IsEmpty() const
+{
+	return MeshDescription.IsEmpty();
+}
+
+void UMeshDescriptionBase::ReserveNewVertices(int32 NumberOfNewVertices)
+{
+	MeshDescription.ReserveNewVertices(NumberOfNewVertices);
+}
+
+FVertexID UMeshDescriptionBase::CreateVertex()
+{
+	return MeshDescription.CreateVertex();
+}
+
+void UMeshDescriptionBase::CreateVertexWithID(FVertexID VertexID)
+{
+	if (!MeshDescription.IsVertexValid(VertexID))
+	{
+		MeshDescription.CreateVertexWithID(VertexID);
+	}
+	else
+	{
+		UE_LOG(LogMeshDescription, Warning, TEXT("CreateVertexWithID: VertexID %d already exists."), VertexID.GetValue());
+	}
+}
+
+void UMeshDescriptionBase::DeleteVertex(FVertexID VertexID)
+{
+	if (MeshDescription.IsVertexValid(VertexID))
+	{
+		MeshDescription.DeleteVertex(VertexID);
+	}
+	else
+	{
+		UE_LOG(LogMeshDescription, Warning, TEXT("DeleteVertex: VertexID %d doesn't exist."), VertexID.GetValue());
+	}
+}
+
+bool UMeshDescriptionBase::IsVertexValid(FVertexID VertexID) const
+{
+	return MeshDescription.IsVertexValid(VertexID);
+}
+
+void UMeshDescriptionBase::ReserveNewVertexInstances(int32 NumberOfNewVertexInstances)
+{
+	MeshDescription.ReserveNewEdges(NumberOfNewVertexInstances);
+}
+
+FVertexInstanceID UMeshDescriptionBase::CreateVertexInstance(FVertexID VertexID)
+{
+	if (MeshDescription.IsVertexValid(VertexID))
+	{
+		return MeshDescription.CreateVertexInstance(VertexID);
+	}
+	else
+	{
+		UE_LOG(LogMeshDescription, Warning, TEXT("CreateVertexInstance: VertexID %d doesn't exist."), VertexID.GetValue());
+		return FVertexInstanceID::Invalid;
+	}
+}
+
+void UMeshDescriptionBase::CreateVertexInstanceWithID(FVertexInstanceID VertexInstanceID, FVertexID VertexID)
+{
+	if (MeshDescription.IsVertexValid(VertexID))
+	{
+		if (!MeshDescription.IsVertexInstanceValid(VertexInstanceID))
+		{
+			MeshDescription.CreateVertexInstanceWithID(VertexInstanceID, VertexID);
+		}
+		else
+		{
+			UE_LOG(LogMeshDescription, Warning, TEXT("CreateVertexInstanceWithID: VertexInstanceID %d already exists."), VertexInstanceID.GetValue());
+		}
+	}
+	else
+	{
+		UE_LOG(LogMeshDescription, Warning, TEXT("CreateVertexInstanceWithID: VertexID %d doesn't exist."), VertexID.GetValue());
+	}
+}
+
+void UMeshDescriptionBase::DeleteVertexInstance(FVertexInstanceID VertexInstanceID, TArray<FVertexID>& OrphanedVertices)
+{
+	if (MeshDescription.IsVertexInstanceValid(VertexInstanceID))
+	{
+		MeshDescription.DeleteVertexInstance(VertexInstanceID, &OrphanedVertices);
+	}
+	else
+	{
+		UE_LOG(LogMeshDescription, Warning, TEXT("DeleteVertexInstance: VertexInstanceID %d doesn't exist."), VertexInstanceID.GetValue());
+	}
+}
+
+bool UMeshDescriptionBase::IsVertexInstanceValid(FVertexInstanceID VertexInstanceID) const
+{
+	return MeshDescription.IsVertexInstanceValid(VertexInstanceID);
+}
+
+void UMeshDescriptionBase::ReserveNewEdges(int32 NumberOfNewEdges)
+{
+	MeshDescription.ReserveNewEdges(NumberOfNewEdges);
+}
+
+FEdgeID UMeshDescriptionBase::CreateEdge(FVertexID VertexID0, FVertexID VertexID1)
+{
+	if (!MeshDescription.IsVertexValid(VertexID0))
+	{
+		UE_LOG(LogMeshDescription, Warning, TEXT("CreateEdge: VertexID %d doesn't exist."), VertexID0.GetValue());
+		return FEdgeID::Invalid;
+	}
+
+	if (!MeshDescription.IsVertexValid(VertexID1))
+	{
+		UE_LOG(LogMeshDescription, Warning, TEXT("CreateEdge: VertexID %d doesn't exist."), VertexID1.GetValue());
+		return FEdgeID::Invalid;
+	}
+
+	return MeshDescription.CreateEdge(VertexID0, VertexID1);
+}
+
+void UMeshDescriptionBase::CreateEdgeWithID(FEdgeID EdgeID, FVertexID VertexID0, FVertexID VertexID1)
+{
+	if (MeshDescription.IsEdgeValid(EdgeID))
+	{
+		UE_LOG(LogMeshDescription, Warning, TEXT("CreateEdgeWithID: EdgeID %d already exists."), EdgeID.GetValue());
+		return;
+	}
+
+	if (!MeshDescription.IsVertexValid(VertexID0))
+	{
+		UE_LOG(LogMeshDescription, Warning, TEXT("CreateEdgeWithID: VertexID %d doesn't exist."), VertexID0.GetValue());
+		return;
+	}
+
+	if (!MeshDescription.IsVertexValid(VertexID1))
+	{
+		UE_LOG(LogMeshDescription, Warning, TEXT("CreateEdgeWithID: VertexID %d doesn't exist."), VertexID1.GetValue());
+		return;
+	}
+
+	MeshDescription.CreateEdgeWithID(EdgeID, VertexID0, VertexID1);
+}
+
+void UMeshDescriptionBase::DeleteEdge(FEdgeID EdgeID, TArray<FVertexID>& OrphanedVertices)
+{
+	if (MeshDescription.IsEdgeValid(EdgeID))
+	{
+		MeshDescription.DeleteEdge(EdgeID, &OrphanedVertices);
+	}
+	else
+	{
+		UE_LOG(LogMeshDescription, Warning, TEXT("DeleteEdge: EdgeID %d doesn't exist."), EdgeID.GetValue());
+	}
+}
+
+bool UMeshDescriptionBase::IsEdgeValid(FEdgeID EdgeID) const
+{
+	return MeshDescription.IsEdgeValid(EdgeID);
+}
+
+void UMeshDescriptionBase::ReserveNewTriangles(int32 NumberOfNewTriangles)
+{
+	MeshDescription.ReserveNewTriangles(NumberOfNewTriangles);
+}
+
+FTriangleID UMeshDescriptionBase::CreateTriangle(FPolygonGroupID PolygonGroupID, const TArray<FVertexInstanceID>& VertexInstanceIDs, TArray<FEdgeID>& NewEdgeIDs)
+{
+	if (!MeshDescription.IsPolygonGroupValid(PolygonGroupID))
+	{
+		UE_LOG(LogMeshDescription, Warning, TEXT("CreateTriangle: PolygonGroupID %d doesn't exist."), PolygonGroupID.GetValue());
+		return FTriangleID::Invalid;
+	}
+
+	return MeshDescription.CreateTriangle(PolygonGroupID, VertexInstanceIDs, &NewEdgeIDs);
+}
+
+void UMeshDescriptionBase::CreateTriangleWithID(FTriangleID TriangleID, FPolygonGroupID PolygonGroupID, const TArray<FVertexInstanceID>& VertexInstanceIDs, TArray<FEdgeID>& NewEdgeIDs)
+{
+	if (MeshDescription.IsTriangleValid(TriangleID))
+	{
+		UE_LOG(LogMeshDescription, Warning, TEXT("CreateTriangleWithID: TriangleID %d already exists."), TriangleID.GetValue());
+		return;
+	}
+
+	if (!MeshDescription.IsPolygonGroupValid(PolygonGroupID))
+	{
+		UE_LOG(LogMeshDescription, Warning, TEXT("CreateTriangleWithID: PolygonGroupID %d doesn't exist."), PolygonGroupID.GetValue());
+		return;
+	}
+
+	MeshDescription.CreateTriangleWithID(TriangleID, PolygonGroupID, VertexInstanceIDs, &NewEdgeIDs);
+}
+
+void UMeshDescriptionBase::DeleteTriangle(FTriangleID TriangleID, TArray<FEdgeID>& OrphanedEdges, TArray<FVertexInstanceID>& OrphanedVertexInstances, TArray<FPolygonGroupID>& OrphanedPolygonGroups)
+{
+	if (!MeshDescription.IsTriangleValid(TriangleID))
+	{
+		UE_LOG(LogMeshDescription, Warning, TEXT("DeleteTriangle: TriangleID %d doesn't exist."), TriangleID.GetValue());
+		return;
+	}
+
+	MeshDescription.DeleteTriangle(TriangleID, &OrphanedEdges, &OrphanedVertexInstances, &OrphanedPolygonGroups);
+}
+
+bool UMeshDescriptionBase::IsTriangleValid(const FTriangleID TriangleID) const
+{
+	return MeshDescription.IsTriangleValid(TriangleID);
+}
+
+void UMeshDescriptionBase::ReserveNewPolygons(const int32 NumberOfNewPolygons)
+{
+	MeshDescription.ReserveNewPolygons(NumberOfNewPolygons);
+}
+
+FPolygonID UMeshDescriptionBase::CreatePolygon(FPolygonGroupID PolygonGroupID, TArray<FVertexInstanceID>& VertexInstanceIDs, TArray<FEdgeID>& NewEdgeIDs)
+{
+	if (!MeshDescription.IsPolygonGroupValid(PolygonGroupID))
+	{
+		UE_LOG(LogMeshDescription, Warning, TEXT("CreatePolygon: PolygonGroupID %d doesn't exist."), PolygonGroupID.GetValue());
+		return FPolygonID::Invalid;
+	}
+
+	return MeshDescription.CreatePolygon(PolygonGroupID, VertexInstanceIDs, &NewEdgeIDs);
+}
+
+void UMeshDescriptionBase::CreatePolygonWithID(FPolygonID PolygonID, FPolygonGroupID PolygonGroupID, TArray<FVertexInstanceID>& VertexInstanceIDs, TArray<FEdgeID>& NewEdgeIDs)
+{
+	if (MeshDescription.IsPolygonValid(PolygonID))
+	{
+		UE_LOG(LogMeshDescription, Warning, TEXT("CreatePolygonWithID: PolygonID %d already exists."), PolygonID.GetValue());
+		return;
+	}
+
+	if (!MeshDescription.IsPolygonGroupValid(PolygonGroupID))
+	{
+		UE_LOG(LogMeshDescription, Warning, TEXT("CreatePolygonWithID: PolygonGroupID %d doesn't exist."), PolygonGroupID.GetValue());
+		return;
+	}
+
+	MeshDescription.CreatePolygonWithID(PolygonID, PolygonGroupID, VertexInstanceIDs, &NewEdgeIDs);
+}
+
+void UMeshDescriptionBase::DeletePolygon(FPolygonID PolygonID, TArray<FEdgeID>& OrphanedEdges, TArray<FVertexInstanceID>& OrphanedVertexInstances, TArray<FPolygonGroupID>& OrphanedPolygonGroups)
+{
+	if (!MeshDescription.IsPolygonValid(PolygonID))
+	{
+		UE_LOG(LogMeshDescription, Warning, TEXT("DeletePolygon: PolygonID %d doesn't exist."), PolygonID.GetValue());
+		return;
+	}
+
+	MeshDescription.DeletePolygon(PolygonID, &OrphanedEdges, &OrphanedVertexInstances, &OrphanedPolygonGroups);
+}
+
+bool UMeshDescriptionBase::IsPolygonValid(FPolygonID PolygonID) const
+{
+	return MeshDescription.IsPolygonValid(PolygonID);
+}
+
+void UMeshDescriptionBase::ReserveNewPolygonGroups(int32 NumberOfNewPolygonGroups)
+{
+	MeshDescription.ReserveNewPolygonGroups(NumberOfNewPolygonGroups);
+}
+
+FPolygonGroupID UMeshDescriptionBase::CreatePolygonGroup()
+{
+	return MeshDescription.CreatePolygonGroup();
+}
+
+void UMeshDescriptionBase::CreatePolygonGroupWithID(FPolygonGroupID PolygonGroupID)
+{
+	if (MeshDescription.IsPolygonGroupValid(PolygonGroupID))
+	{
+		UE_LOG(LogMeshDescription, Warning, TEXT("CreatePolygonGroupWithID: PolygonGroupID %d already exists."), PolygonGroupID.GetValue());
+		return;
+	}
+
+	MeshDescription.CreatePolygonGroupWithID(PolygonGroupID);
+}
+
+void UMeshDescriptionBase::DeletePolygonGroup(FPolygonGroupID PolygonGroupID)
+{
+	if (!MeshDescription.IsPolygonGroupValid(PolygonGroupID))
+	{
+		UE_LOG(LogMeshDescription, Warning, TEXT("DeletePolygonGroup: FPolygonGroupID %d doesn't exist."), PolygonGroupID.GetValue());
+		return;
+	}
+
+	MeshDescription.DeletePolygonGroup(PolygonGroupID);
+}
+
+bool UMeshDescriptionBase::IsPolygonGroupValid(FPolygonGroupID PolygonGroupID) const
+{
+	return MeshDescription.IsPolygonGroupValid(PolygonGroupID);
+}
+
+
+//////////////////////////////////////////////////////////////////////
+// Vertex operations
+
+bool UMeshDescriptionBase::IsVertexOrphaned(FVertexID VertexID) const
+{
+	if (!MeshDescription.IsVertexValid(VertexID))
+	{
+		UE_LOG(LogMeshDescription, Warning, TEXT("IsVertexOrphaned: VertexID %d doesn't exist."), VertexID.GetValue());
+		return false;
+	}
+
+	return MeshDescription.IsVertexOrphaned(VertexID);
+}
+
+FEdgeID UMeshDescriptionBase::GetVertexPairEdge(FVertexID VertexID0, FVertexID VertexID1) const
+{
+	if (!MeshDescription.IsVertexValid(VertexID0))
+	{
+		UE_LOG(LogMeshDescription, Warning, TEXT("GetVertexPairEdge: VertexID %d doesn't exist."), VertexID0.GetValue());
+		return FEdgeID::Invalid;
+	}
+
+	if (!MeshDescription.IsVertexValid(VertexID1))
+	{
+		UE_LOG(LogMeshDescription, Warning, TEXT("GetVertexPairEdge: VertexID %d doesn't exist."), VertexID1.GetValue());
+		return FEdgeID::Invalid;
+	}
+
+	return MeshDescription.GetVertexPairEdge(VertexID0, VertexID1);
+}
+
+void UMeshDescriptionBase::GetVertexConnectedEdges(FVertexID VertexID, TArray<FEdgeID>& OutEdgeIDs) const
+{
+	if (!MeshDescription.IsVertexValid(VertexID))
+	{
+		UE_LOG(LogMeshDescription, Warning, TEXT("GetVertexConnectedEdges: VertexID %d doesn't exist."), VertexID.GetValue());
+		return;
+	}
+
+	OutEdgeIDs = MeshDescription.GetVertexConnectedEdges(VertexID);
+}
+
+int32 UMeshDescriptionBase::GetNumVertexConnectedEdges(FVertexID VertexID) const
+{
+	if (!MeshDescription.IsVertexValid(VertexID))
+	{
+		UE_LOG(LogMeshDescription, Warning, TEXT("GetNumVertexConnectedEdges: VertexID %d doesn't exist."), VertexID.GetValue());
+		return 0;
+	}
+
+	return MeshDescription.GetNumVertexConnectedEdges(VertexID);
+}
+
+void UMeshDescriptionBase::GetVertexVertexInstances(FVertexID VertexID, TArray<FVertexInstanceID>& OutVertexInstanceIDs) const
+{
+	if (!MeshDescription.IsVertexValid(VertexID))
+	{
+		UE_LOG(LogMeshDescription, Warning, TEXT("GetVertexVertexInstances: VertexID %d doesn't exist."), VertexID.GetValue());
+		return;
+	}
+
+	OutVertexInstanceIDs = MeshDescription.GetVertexVertexInstances(VertexID);
+}
+
+int32 UMeshDescriptionBase::GetNumVertexVertexInstances(FVertexID VertexID) const
+{
+	if (!MeshDescription.IsVertexValid(VertexID))
+	{
+		UE_LOG(LogMeshDescription, Warning, TEXT("GetNumVertexVertexInstances: VertexID %d doesn't exist."), VertexID.GetValue());
+		return 0;
+	}
+
+	return MeshDescription.GetNumVertexVertexInstances(VertexID);
+}
+
+void UMeshDescriptionBase::GetVertexConnectedTriangles(FVertexID VertexID, TArray<FTriangleID>& OutConnectedTriangleIDs) const
+{
+	if (!MeshDescription.IsVertexValid(VertexID))
+	{
+		UE_LOG(LogMeshDescription, Warning, TEXT("GetVertexConnectedTriangles: VertexID %d doesn't exist."), VertexID.GetValue());
+		return;
+	}
+
+	MeshDescription.GetVertexConnectedTriangles(VertexID, OutConnectedTriangleIDs);
+}
+
+int32 UMeshDescriptionBase::GetNumVertexConnectedTriangles(FVertexID VertexID) const
+{
+	if (!MeshDescription.IsVertexValid(VertexID))
+	{
+		UE_LOG(LogMeshDescription, Warning, TEXT("GetNumVertexConnectedTriangles: VertexID %d doesn't exist."), VertexID.GetValue());
+		return 0;
+	}
+
+	return MeshDescription.GetNumVertexConnectedTriangles(VertexID);
+}
+
+void UMeshDescriptionBase::GetVertexConnectedPolygons(FVertexID VertexID, TArray<FPolygonID>& OutConnectedPolygonIDs) const
+{
+	if (!MeshDescription.IsVertexValid(VertexID))
+	{
+		UE_LOG(LogMeshDescription, Warning, TEXT("GetVertexConnectedPolygons: VertexID %d doesn't exist."), VertexID.GetValue());
+		return;
+	}
+
+	MeshDescription.GetVertexConnectedPolygons(VertexID, OutConnectedPolygonIDs);
+}
+
+int32 UMeshDescriptionBase::GetNumVertexConnectedPolygons(FVertexID VertexID) const
+{
+	if (!MeshDescription.IsVertexValid(VertexID))
+	{
+		UE_LOG(LogMeshDescription, Warning, TEXT("GetNumVertexConnectedPolygons: VertexID %d doesn't exist."), VertexID.GetValue());
+		return 0;
+	}
+
+	return MeshDescription.GetNumVertexConnectedPolygons(VertexID);
+}
+
+void UMeshDescriptionBase::GetVertexAdjacentVertices(FVertexID VertexID, TArray<FVertexID>& OutAdjacentVertexIDs) const
+{
+	if (!MeshDescription.IsVertexValid(VertexID))
+	{
+		UE_LOG(LogMeshDescription, Warning, TEXT("GetVertexAdjacentVertices: VertexID %d doesn't exist."), VertexID.GetValue());
+		return;
+	}
+
+	MeshDescription.GetVertexAdjacentVertices(VertexID, OutAdjacentVertexIDs);
+}
+
+FVector UMeshDescriptionBase::GetVertexPosition(FVertexID VertexID) const
+{
+	if (!MeshDescription.IsVertexValid(VertexID))
+	{
+		UE_LOG(LogMeshDescription, Warning, TEXT("GetVertexAttribute: VertexID %d doesn't exist."), VertexID.GetValue());
+		return FVector::ZeroVector;
+	}
+
+	if (!MeshDescription.VertexAttributes().HasAttribute(MeshAttribute::Vertex::Position))
+	{
+		UE_LOG(LogMeshDescription, Warning, TEXT("GetVertexAttribute: VertexAttribute Position doesn't exist."));
+		return FVector::ZeroVector;
+	}
+
+	return MeshDescription.VertexAttributes().GetAttribute<FVector>(VertexID, MeshAttribute::Vertex::Position);
+}
+
+void UMeshDescriptionBase::SetVertexPosition(FVertexID VertexID, const FVector& Position)
+{
+	if (!MeshDescription.IsVertexValid(VertexID))
+	{
+		UE_LOG(LogMeshDescription, Warning, TEXT("SetVertexAttribute: VertexID %d doesn't exist."), VertexID.GetValue());
+		return;
+	}
+
+	if (!MeshDescription.VertexAttributes().HasAttribute(MeshAttribute::Vertex::Position))
+	{
+		UE_LOG(LogMeshDescription, Warning, TEXT("SetVertexAttribute: VertexAttribute Position doesn't exist."));
+		return;
+	}
+
+	MeshDescription.VertexAttributes().SetAttribute(VertexID, MeshAttribute::Vertex::Position, 0, Position);
+}
+
+
+//////////////////////////////////////////////////////////////////////
+// Vertex instance operations
+
+FVertexID UMeshDescriptionBase::GetVertexInstanceVertex(FVertexInstanceID VertexInstanceID) const
+{
+	if (!MeshDescription.IsVertexInstanceValid(VertexInstanceID))
+	{
+		UE_LOG(LogMeshDescription, Warning, TEXT("GetVertexInstanceVertex: VertexInstanceID %d doesn't exist."), VertexInstanceID.GetValue());
+		return FVertexID::Invalid;
+	}
+
+	return MeshDescription.GetVertexInstanceVertex(VertexInstanceID);
+}
+
+FEdgeID UMeshDescriptionBase::GetVertexInstancePairEdge(FVertexInstanceID VertexInstanceID0, FVertexInstanceID VertexInstanceID1) const
+{
+	if (!MeshDescription.IsVertexInstanceValid(VertexInstanceID0))
+	{
+		UE_LOG(LogMeshDescription, Warning, TEXT("GetVertexInstancePairEdge: VertexInstanceID %d doesn't exist."), VertexInstanceID0.GetValue());
+		return FEdgeID::Invalid;
+	}
+
+	if (!MeshDescription.IsVertexInstanceValid(VertexInstanceID1))
+	{
+		UE_LOG(LogMeshDescription, Warning, TEXT("GetVertexInstancePairEdge: VertexInstanceID %d doesn't exist."), VertexInstanceID1.GetValue());
+		return FEdgeID::Invalid;
+	}
+
+	return MeshDescription.GetVertexInstancePairEdge(VertexInstanceID0, VertexInstanceID1);
+}
+
+void UMeshDescriptionBase::GetVertexInstanceConnectedTriangles(FVertexInstanceID VertexInstanceID, TArray<FTriangleID>& OutConnectedTriangleIDs) const
+{
+	if (!MeshDescription.IsVertexInstanceValid(VertexInstanceID))
+	{
+		UE_LOG(LogMeshDescription, Warning, TEXT("GetVertexInstanceConnectedTriangles: VertexInstanceID %d doesn't exist."), VertexInstanceID.GetValue());
+		return;
+	}
+
+	OutConnectedTriangleIDs = MeshDescription.GetVertexInstanceConnectedTriangles(VertexInstanceID);
+}
+
+int32 UMeshDescriptionBase::GetNumVertexInstanceConnectedTriangles(FVertexInstanceID VertexInstanceID) const
+{
+	if (!MeshDescription.IsVertexInstanceValid(VertexInstanceID))
+	{
+		UE_LOG(LogMeshDescription, Warning, TEXT("GetNumVertexInstanceConnectedTriangles: VertexInstanceID %d doesn't exist."), VertexInstanceID.GetValue());
+		return 0;
+	}
+
+	return MeshDescription.GetNumVertexInstanceConnectedTriangles(VertexInstanceID);
+}
+
+void UMeshDescriptionBase::GetVertexInstanceConnectedPolygons(FVertexInstanceID VertexInstanceID, TArray<FPolygonID>& OutConnectedPolygonIDs) const
+{
+	if (!MeshDescription.IsVertexInstanceValid(VertexInstanceID))
+	{
+		UE_LOG(LogMeshDescription, Warning, TEXT("GetVertexInstanceConnectedPolygons: VertexInstanceID %d doesn't exist."), VertexInstanceID.GetValue());
+		return;
+	}
+
+	MeshDescription.GetVertexInstanceConnectedPolygons(VertexInstanceID, OutConnectedPolygonIDs);
+}
+
+int32 UMeshDescriptionBase::GetNumVertexInstanceConnectedPolygons(FVertexInstanceID VertexInstanceID) const
+{
+	if (!MeshDescription.IsVertexInstanceValid(VertexInstanceID))
+	{
+		UE_LOG(LogMeshDescription, Warning, TEXT("GetNumVertexInstanceConnectedPolygons: VertexInstanceID %d doesn't exist."), VertexInstanceID.GetValue());
+		return 0;
+	}
+
+	return MeshDescription.GetNumVertexInstanceConnectedPolygons(VertexInstanceID);
+}
+
+
+//////////////////////////////////////////////////////////////////////
+// Edge operations
+
+bool UMeshDescriptionBase::IsEdgeInternal(FEdgeID EdgeID) const
+{
+	if (!MeshDescription.IsEdgeValid(EdgeID))
+	{
+		UE_LOG(LogMeshDescription, Warning, TEXT("IsEdgeInternal: EdgeID %d doesn't exist."), EdgeID.GetValue());
+		return false;
+	}
+
+	return MeshDescription.IsEdgeInternal(EdgeID);
+}
+
+bool UMeshDescriptionBase::IsEdgeInternalToPolygon(FEdgeID EdgeID, FPolygonID PolygonID) const
+{
+	if (!MeshDescription.IsEdgeValid(EdgeID))
+	{
+		UE_LOG(LogMeshDescription, Warning, TEXT("IsEdgeInternalToPolygon: EdgeID %d doesn't exist."), EdgeID.GetValue());
+		return false;
+	}
+
+	if (!MeshDescription.IsPolygonValid(PolygonID))
+	{
+		UE_LOG(LogMeshDescription, Warning, TEXT("IsEdgeInternalToPolygon: PolygonID %d doesn't exist."), PolygonID.GetValue());
+		return false;
+	}
+
+	return MeshDescription.IsEdgeInternalToPolygon(EdgeID, PolygonID);
+}
+
+void UMeshDescriptionBase::GetEdgeConnectedTriangles(FEdgeID EdgeID, TArray<FTriangleID>& OutConnectedTriangleIDs) const
+{
+	if (!MeshDescription.IsEdgeValid(EdgeID))
+	{
+		UE_LOG(LogMeshDescription, Warning, TEXT("GetEdgeConnectedTriangles: EdgeID %d doesn't exist."), EdgeID.GetValue());
+		return;
+	}
+
+	OutConnectedTriangleIDs = MeshDescription.GetEdgeConnectedTriangles(EdgeID);
+}
+
+int32 UMeshDescriptionBase::GetNumEdgeConnectedTriangles(FEdgeID EdgeID) const
+{
+	if (!MeshDescription.IsEdgeValid(EdgeID))
+	{
+		UE_LOG(LogMeshDescription, Warning, TEXT("GetNumEdgeConnectedTriangles: EdgeID %d doesn't exist."), EdgeID.GetValue());
+		return 0;
+	}
+
+	return MeshDescription.GetNumEdgeConnectedTriangles(EdgeID);
+}
+
+void UMeshDescriptionBase::GetEdgeConnectedPolygons(FEdgeID EdgeID, TArray<FPolygonID>& OutConnectedPolygonIDs) const
+{
+	if (!MeshDescription.IsEdgeValid(EdgeID))
+	{
+		UE_LOG(LogMeshDescription, Warning, TEXT("GetEdgeConnectedPolygons: EdgeID %d doesn't exist."), EdgeID.GetValue());
+		return;
+	}
+
+	MeshDescription.GetEdgeConnectedPolygons(EdgeID, OutConnectedPolygonIDs);
+}
+
+int32 UMeshDescriptionBase::GetNumEdgeConnectedPolygons(FEdgeID EdgeID) const
+{
+	if (!MeshDescription.IsEdgeValid(EdgeID))
+	{
+		UE_LOG(LogMeshDescription, Warning, TEXT("GetNumEdgeConnectedPolygons: EdgeID %d doesn't exist."), EdgeID.GetValue());
+		return 0;
+	}
+
+	return MeshDescription.GetNumEdgeConnectedPolygons(EdgeID);
+}
+
+FVertexID UMeshDescriptionBase::GetEdgeVertex(FEdgeID EdgeID, int32 VertexNumber) const
+{
+	if (!MeshDescription.IsEdgeValid(EdgeID))
+	{
+		UE_LOG(LogMeshDescription, Warning, TEXT("GetEdgeVertex: EdgeID %d doesn't exist."), EdgeID.GetValue());
+		return FVertexID::Invalid;
+	}
+
+	if (VertexNumber != 0 && VertexNumber != 1)
+	{
+		UE_LOG(LogMeshDescription, Warning, TEXT("GetEdgeVertex: invalid vertex number %d."), VertexNumber);
+		return FVertexID::Invalid;
+	}
+
+	return MeshDescription.GetEdgeVertex(EdgeID, VertexNumber);
+}
+
+void UMeshDescriptionBase::GetEdgeVertices(const FEdgeID EdgeID, TArray<FVertexID>& OutVertexIDs) const
+{
+	if (!MeshDescription.IsEdgeValid(EdgeID))
+	{
+		UE_LOG(LogMeshDescription, Warning, TEXT("GetEdgeVertices: EdgeID %d doesn't exist."), EdgeID.GetValue());
+		return;
+	}
+
+	Algo::Copy(MeshDescription.GetEdgeVertices(EdgeID), OutVertexIDs);
+}
+
+
+//////////////////////////////////////////////////////////////////////
+// Triangle operations
+
+FPolygonID UMeshDescriptionBase::GetTrianglePolygon(FTriangleID TriangleID) const
+{
+	if (!MeshDescription.IsTriangleValid(TriangleID))
+	{
+		UE_LOG(LogMeshDescription, Warning, TEXT("GetTrianglePolygon: TriangleID %d doesn't exist."), TriangleID.GetValue());
+		return FPolygonID::Invalid;
+	}
+
+	return MeshDescription.GetTrianglePolygon(TriangleID);
+}
+
+FPolygonGroupID UMeshDescriptionBase::GetTrianglePolygonGroup(FTriangleID TriangleID) const
+{
+	if (!MeshDescription.IsTriangleValid(TriangleID))
+	{
+		UE_LOG(LogMeshDescription, Warning, TEXT("GetTrianglePolygonGroup: TriangleID %d doesn't exist."), TriangleID.GetValue());
+		return FPolygonGroupID::Invalid;
+	}
+
+	return MeshDescription.GetTrianglePolygonGroup(TriangleID);
+}
+
+bool UMeshDescriptionBase::IsTrianglePartOfNgon(FTriangleID TriangleID) const
+{
+	if (!MeshDescription.IsTriangleValid(TriangleID))
+	{
+		UE_LOG(LogMeshDescription, Warning, TEXT("IsTrianglePartOfNgon: TriangleID %d doesn't exist."), TriangleID.GetValue());
+		return false;
+	}
+
+	return MeshDescription.IsTrianglePartOfNgon(TriangleID);
+}
+
+void UMeshDescriptionBase::GetTriangleVertexInstances(FTriangleID TriangleID, TArray<FVertexInstanceID>& OutVertexInstanceIDs) const
+{
+	if (!MeshDescription.IsTriangleValid(TriangleID))
+	{
+		UE_LOG(LogMeshDescription, Warning, TEXT("GetTriangleVertexInstances: TriangleID %d doesn't exist."), TriangleID.GetValue());
+		return;
+	}
+
+	Algo::Copy(MeshDescription.GetTriangleVertexInstances(TriangleID), OutVertexInstanceIDs);
+}
+
+FVertexInstanceID UMeshDescriptionBase::GetTriangleVertexInstance(FTriangleID TriangleID, int32 Index) const
+{
+	if (!MeshDescription.IsTriangleValid(TriangleID))
+	{
+		UE_LOG(LogMeshDescription, Warning, TEXT("GetTriangleVertexInstance: TriangleID %d doesn't exist."), TriangleID.GetValue());
+		return FVertexInstanceID::Invalid;
+	}
+
+	if (Index < 0 || Index > 2)
+	{
+		UE_LOG(LogMeshDescription, Warning, TEXT("GetTriangleVertexInstance: invalid vertex index %d."), Index);
+		return FVertexInstanceID::Invalid;
+	}
+
+	return MeshDescription.GetTriangleVertexInstance(TriangleID, Index);
+}
+
+void UMeshDescriptionBase::GetTriangleVertices(FTriangleID TriangleID, TArray<FVertexID>& OutVertexIDs) const
+{
+	if (!MeshDescription.IsTriangleValid(TriangleID))
+	{
+		UE_LOG(LogMeshDescription, Warning, TEXT("GetTriangleVertices: TriangleID %d doesn't exist."), TriangleID.GetValue());
+		return;
+	}
+
+	OutVertexIDs.SetNumUninitialized(3);
+	MeshDescription.GetTriangleVertices(TriangleID, OutVertexIDs);
+}
+
+void UMeshDescriptionBase::GetTriangleEdges(FTriangleID TriangleID, TArray<FEdgeID>& OutEdgeIDs) const
+{
+	if (!MeshDescription.IsTriangleValid(TriangleID))
+	{
+		UE_LOG(LogMeshDescription, Warning, TEXT("GetTriangleEdges: TriangleID %d doesn't exist."), TriangleID.GetValue());
+		return;
+	}
+
+	OutEdgeIDs.SetNumUninitialized(3);
+	MeshDescription.GetTriangleEdges(TriangleID, OutEdgeIDs);
+}
+
+void UMeshDescriptionBase::GetTriangleAdjacentTriangles(FTriangleID TriangleID, TArray<FTriangleID>& OutTriangleIDs) const
+{
+	if (!MeshDescription.IsTriangleValid(TriangleID))
+	{
+		UE_LOG(LogMeshDescription, Warning, TEXT("GetTriangleAdjacentTriangles: TriangleID %d doesn't exist."), TriangleID.GetValue());
+		return;
+	}
+
+	MeshDescription.GetTriangleAdjacentTriangles(TriangleID, OutTriangleIDs);
+}
+
+FVertexInstanceID UMeshDescriptionBase::GetVertexInstanceForTriangleVertex(FTriangleID TriangleID, FVertexID VertexID) const
+{
+	if (!MeshDescription.IsTriangleValid(TriangleID))
+	{
+		UE_LOG(LogMeshDescription, Warning, TEXT("GetVertexInstanceForTriangleVertex: TriangleID %d doesn't exist."), TriangleID.GetValue());
+		return FVertexInstanceID::Invalid;
+	}
+
+	return MeshDescription.GetVertexInstanceForTriangleVertex(TriangleID, VertexID);
+}
+
+
+//////////////////////////////////////////////////////////////////////
+// Polygon operations
+
+void UMeshDescriptionBase::GetPolygonTriangles(FPolygonID PolygonID, TArray<FTriangleID>& OutTriangleIDs) const
+{
+	if (!MeshDescription.IsPolygonValid(PolygonID))
+	{
+		UE_LOG(LogMeshDescription, Warning, TEXT("GetPolygonTriangles: PolygonID %d doesn't exist."), PolygonID.GetValue());
+		return;
+	}
+
+	OutTriangleIDs = MeshDescription.GetPolygonTriangleIDs(PolygonID);
+}
+
+int32 UMeshDescriptionBase::GetNumPolygonTriangles(FPolygonID PolygonID) const
+{
+	if (!MeshDescription.IsPolygonValid(PolygonID))
+	{
+		UE_LOG(LogMeshDescription, Warning, TEXT("GetNumPolygonTriangles: PolygonID %d doesn't exist."), PolygonID.GetValue());
+		return 0;
+	}
+
+	return MeshDescription.GetNumPolygonTriangles(PolygonID);
+}
+
+void UMeshDescriptionBase::GetPolygonVertexInstances(FPolygonID PolygonID, TArray<FVertexInstanceID>& OutVertexInstanceIDs) const
+{
+	if (!MeshDescription.IsPolygonValid(PolygonID))
+	{
+		UE_LOG(LogMeshDescription, Warning, TEXT("GetPolygonVertexInstances: PolygonID %d doesn't exist."), PolygonID.GetValue());
+		return;
+	}
+
+	OutVertexInstanceIDs = MeshDescription.GetPolygonVertexInstances(PolygonID);
+}
+
+int32 UMeshDescriptionBase::GetNumPolygonVertices(FPolygonID PolygonID) const
+{
+	if (!MeshDescription.IsPolygonValid(PolygonID))
+	{
+		UE_LOG(LogMeshDescription, Warning, TEXT("GetNumPolygonVertices: PolygonID %d doesn't exist."), PolygonID.GetValue());
+		return 0;
+	}
+
+	return MeshDescription.GetNumPolygonVertices(PolygonID);
+}
+
+void UMeshDescriptionBase::GetPolygonVertices(FPolygonID PolygonID, TArray<FVertexID>& OutVertexIDs) const
+{
+	if (!MeshDescription.IsPolygonValid(PolygonID))
+	{
+		UE_LOG(LogMeshDescription, Warning, TEXT("GetPolygonVertices: PolygonID %d doesn't exist."), PolygonID.GetValue());
+		return;
+	}
+
+	MeshDescription.GetPolygonVertices(PolygonID, OutVertexIDs);
+}
+
+void UMeshDescriptionBase::GetPolygonPerimeterEdges(FPolygonID PolygonID, TArray<FEdgeID>& OutEdgeIDs) const
+{
+	if (!MeshDescription.IsPolygonValid(PolygonID))
+	{
+		UE_LOG(LogMeshDescription, Warning, TEXT("GetPolygonPerimeterEdges: PolygonID %d doesn't exist."), PolygonID.GetValue());
+		return;
+	}
+
+	MeshDescription.GetPolygonPerimeterEdges(PolygonID, OutEdgeIDs);
+}
+
+void UMeshDescriptionBase::GetPolygonInternalEdges(FPolygonID PolygonID, TArray<FEdgeID>& OutEdgeIDs) const
+{
+	if (!MeshDescription.IsPolygonValid(PolygonID))
+	{
+		UE_LOG(LogMeshDescription, Warning, TEXT("GetPolygonInternalEdges: PolygonID %d doesn't exist."), PolygonID.GetValue());
+		return;
+	}
+
+	MeshDescription.GetPolygonInternalEdges(PolygonID, OutEdgeIDs);
+}
+
+int32 UMeshDescriptionBase::GetNumPolygonInternalEdges(FPolygonID PolygonID) const
+{
+	if (!MeshDescription.IsPolygonValid(PolygonID))
+	{
+		UE_LOG(LogMeshDescription, Warning, TEXT("GetNumPolygonInternalEdges: PolygonID %d doesn't exist."), PolygonID.GetValue());
+		return 0;
+	}
+
+	return MeshDescription.GetNumPolygonInternalEdges(PolygonID);
+}
+
+void UMeshDescriptionBase::GetPolygonAdjacentPolygons(FPolygonID PolygonID, TArray<FPolygonID>& OutPolygonIDs) const
+{
+	if (!MeshDescription.IsPolygonValid(PolygonID))
+	{
+		UE_LOG(LogMeshDescription, Warning, TEXT("GetPolygonAdjacentPolygons: PolygonID %d doesn't exist."), PolygonID.GetValue());
+		return;
+	}
+
+	MeshDescription.GetPolygonAdjacentPolygons(PolygonID, OutPolygonIDs);
+}
+
+FPolygonGroupID UMeshDescriptionBase::GetPolygonPolygonGroup(FPolygonID PolygonID) const
+{
+	if (!MeshDescription.IsPolygonValid(PolygonID))
+	{
+		UE_LOG(LogMeshDescription, Warning, TEXT("GetPolygonPolygonGroup: PolygonID %d doesn't exist."), PolygonID.GetValue());
+		return FPolygonGroupID::Invalid;
+	}
+
+	return MeshDescription.GetPolygonPolygonGroup(PolygonID);
+}
+
+FVertexInstanceID UMeshDescriptionBase::GetVertexInstanceForPolygonVertex(FPolygonID PolygonID, FVertexID VertexID) const
+{
+	if (!MeshDescription.IsPolygonValid(PolygonID))
+	{
+		UE_LOG(LogMeshDescription, Warning, TEXT("GetVertexInstanceForPolygonVertex: PolygonID %d doesn't exist."), PolygonID.GetValue());
+		return FVertexInstanceID::Invalid;
+	}
+
+	return MeshDescription.GetVertexInstanceForPolygonVertex(PolygonID, VertexID);
+}
+
+void UMeshDescriptionBase::SetPolygonVertexInstance(FPolygonID PolygonID, int32 PerimeterIndex, FVertexInstanceID VertexInstanceID)
+{
+	if (!MeshDescription.IsPolygonValid(PolygonID))
+	{
+		UE_LOG(LogMeshDescription, Warning, TEXT("SetPolygonVertexInstance: PolygonID %d doesn't exist."), PolygonID.GetValue());
+		return;
+	}
+
+	if (PerimeterIndex >= MeshDescription.GetNumPolygonVertices(PolygonID))
+	{
+		UE_LOG(LogMeshDescription, Warning, TEXT("SetPolygonVertexInstance: Out of range vertex index %d."), PerimeterIndex);
+		return;
+	}
+
+	if (!MeshDescription.IsVertexInstanceValid(VertexInstanceID))
+	{
+		UE_LOG(LogMeshDescription, Warning, TEXT("SetPolygonVertexInstance: VertexInstanceID %d doesn't exist."), VertexInstanceID.GetValue());
+		return;
+	}
+
+	MeshDescription.SetPolygonVertexInstance(PolygonID, PerimeterIndex, VertexInstanceID);
+}
+
+void UMeshDescriptionBase::SetPolygonPolygonGroup(FPolygonID PolygonID, FPolygonGroupID PolygonGroupID)
+{
+	if (!MeshDescription.IsPolygonValid(PolygonID))
+	{
+		UE_LOG(LogMeshDescription, Warning, TEXT("SetPolygonPolygonGroup: PolygonID %d doesn't exist."), PolygonID.GetValue());
+		return;
+	}
+
+	if (!MeshDescription.IsPolygonGroupValid(PolygonGroupID))
+	{
+		UE_LOG(LogMeshDescription, Warning, TEXT("SetPolygonPolygonGroup: PolygonGroupID %d doesn't exist."), PolygonGroupID.GetValue());
+		return;
+	}
+
+	MeshDescription.SetPolygonPolygonGroup(PolygonID, PolygonGroupID);
+}
+
+void UMeshDescriptionBase::ReversePolygonFacing(FPolygonID PolygonID)
+{
+	if (!MeshDescription.IsPolygonValid(PolygonID))
+	{
+		UE_LOG(LogMeshDescription, Warning, TEXT("ReversePolygonFacing: PolygonID %d doesn't exist."), PolygonID.GetValue());
+		return;
+	}
+
+	MeshDescription.ReversePolygonFacing(PolygonID);
+}
+
+void UMeshDescriptionBase::ComputePolygonTriangulation(FPolygonID PolygonID)
+{
+	if (!MeshDescription.IsPolygonValid(PolygonID))
+	{
+		UE_LOG(LogMeshDescription, Warning, TEXT("ComputePolygonTriangulation: PolygonID %d doesn't exist."), PolygonID.GetValue());
+		return;
+	}
+
+	MeshDescription.ComputePolygonTriangulation(PolygonID);
+}
+
+
+//////////////////////////////////////////////////////////////////////
+// Polygon group operations
+
+void UMeshDescriptionBase::GetPolygonGroupPolygons(FPolygonGroupID PolygonGroupID, TArray<FPolygonID>& OutPolygonIDs) const
+{
+	if (!MeshDescription.IsPolygonGroupValid(PolygonGroupID))
+	{
+		UE_LOG(LogMeshDescription, Warning, TEXT("GetPolygonGroupPolygons: PolygonGroupID %d doesn't exist."), PolygonGroupID.GetValue());
+		return;
+	}
+
+	OutPolygonIDs = MeshDescription.GetPolygonGroupPolygons(PolygonGroupID);
+}
+
+int32 UMeshDescriptionBase::GetNumPolygonGroupPolygons(FPolygonGroupID PolygonGroupID) const
+{
+	if (!MeshDescription.IsPolygonGroupValid(PolygonGroupID))
+	{
+		UE_LOG(LogMeshDescription, Warning, TEXT("GetNumPolygonGroupPolygons: PolygonGroupID %d doesn't exist."), PolygonGroupID.GetValue());
+		return 0;
+	}
+
+	return MeshDescription.GetNumPolygonGroupPolygons(PolygonGroupID);
+}
diff --git a/Engine/Source/Runtime/MeshDescription/Public/MeshAttributeArray.h b/Engine/Source/Runtime/MeshDescription/Public/MeshAttributeArray.h
--- a/Engine/Source/Runtime/MeshDescription/Public/MeshAttributeArray.h
+++ b/Engine/Source/Runtime/MeshDescription/Public/MeshAttributeArray.h
@@ -5,6 +5,7 @@
 #include "CoreMinimal.h"
 #include "MeshTypes.h"
 #include "MeshElementRemappings.h"
+#include "Misc/TVariant.h"
 #include "UObject/ReleaseObjectVersion.h"
 
 
@@ -56,6 +57,13 @@ template <typename T, typename... Ts> struct TTupleElement<0, TTuple<T, Ts...>>
 template <typename... Ts> struct TTupleElement<0, TTuple<Ts...>> { static_assert( sizeof...( Ts ) > 0, "TTuple element index out of range" ); };
 
 
+/**
+ * Helper template which generates a TVariant of all supported attribute types.
+ */
+template <typename Tuple> struct TVariantFromTuple;
+template <typename... Ts> struct TVariantFromTuple<TTuple<Ts...>> { using Type = TVariant<FEmptyVariantState, Ts...>; };
+
+
 /**
  * Class which implements a function jump table to be automatically generated at compile time.
  * This is used by TAttributesSet to provide O(1) dispatch by attribute type at runtime.
diff --git a/Engine/Source/Runtime/MeshDescription/Public/MeshAttributes.h b/Engine/Source/Runtime/MeshDescription/Public/MeshAttributes.h
--- a/Engine/Source/Runtime/MeshDescription/Public/MeshAttributes.h
+++ b/Engine/Source/Runtime/MeshDescription/Public/MeshAttributes.h
@@ -3,46 +3,49 @@
 #pragma once
 
 #include "CoreMinimal.h"
+#include "MeshDescription.h"
 
-// @todo mesh description: for now, these attributes are hardcoded in the MeshDescription module.
-// It would be better if these were defined in separate modules, one per mesh type which uses MeshDescription.
-// This would give an extensible and modular structure whereby a mesh type defines its valid attributes and adapters.
 namespace MeshAttribute
 {
 	namespace Vertex
 	{
 		extern MESHDESCRIPTION_API const FName Position;
-		extern MESHDESCRIPTION_API const FName CornerSharpness;
 	}
+}
 
-	namespace VertexInstance
-	{
-		extern MESHDESCRIPTION_API const FName TextureCoordinate;
-		extern MESHDESCRIPTION_API const FName Normal;
-		extern MESHDESCRIPTION_API const FName Tangent;
-		extern MESHDESCRIPTION_API const FName BinormalSign;
-		extern MESHDESCRIPTION_API const FName Color;
-	}
 
-	namespace Edge
-	{
-		extern MESHDESCRIPTION_API const FName IsHard;
-		extern MESHDESCRIPTION_API const FName IsUVSeam;
-		extern MESHDESCRIPTION_API const FName CreaseSharpness;
-	}
+class MESHDESCRIPTION_API FMeshAttributes
+{
+public:
+	explicit FMeshAttributes(FMeshDescription& InMeshDescription)
+		: MeshDescription(InMeshDescription)
+	{}
 
-	namespace Polygon
-	{
-		extern MESHDESCRIPTION_API const FName Normal;
-		extern MESHDESCRIPTION_API const FName Tangent;
-		extern MESHDESCRIPTION_API const FName Binormal;
-		extern MESHDESCRIPTION_API const FName Center;
-	}
+	virtual ~FMeshAttributes() = default;
 
-	namespace PolygonGroup
-	{
-		extern MESHDESCRIPTION_API const FName ImportedMaterialSlotName;
-		extern MESHDESCRIPTION_API const FName EnableCollision;
-		extern MESHDESCRIPTION_API const FName CastShadow;
-	}
-}
+	virtual void Register();
+
+	/** Accessors for cached vertex position array */
+	TVertexAttributesRef<FVector> GetVertexPositions() { return MeshDescription.VertexAttributes().GetAttributesRef<FVector>(MeshAttribute::Vertex::Position); }
+	TVertexAttributesConstRef<FVector> GetVertexPositions() const { return MeshDescription.VertexAttributes().GetAttributesRef<FVector>(MeshAttribute::Vertex::Position); }
+
+protected:
+
+	FMeshDescription& MeshDescription;
+};
+
+
+class FMeshConstAttributes
+{
+public:
+	explicit FMeshConstAttributes(const FMeshDescription& InMeshDescription)
+		: MeshDescription(InMeshDescription)
+	{}
+
+	/** Accessors for cached vertex position array */
+	TVertexAttributesConstRef<FVector> GetVertexPositions() const { return MeshDescription.VertexAttributes().GetAttributesRef<FVector>(MeshAttribute::Vertex::Position); }
+
+protected:
+
+	const FMeshDescription& MeshDescription;
+};
diff --git a/Engine/Source/Runtime/MeshDescription/Public/MeshDescriptionBase.h b/Engine/Source/Runtime/MeshDescription/Public/MeshDescriptionBase.h
new file mode 100644
--- /dev/null
+++ b/Engine/Source/Runtime/MeshDescription/Public/MeshDescriptionBase.h
@@ -0,0 +1,456 @@
+// Copyright 1998-2019 Epic Games, Inc. All Rights Reserved.
+
+#pragma once
+
+#include "MeshDescription.h"
+#include "MeshAttributes.h"
+#include "MeshDescriptionBase.generated.h"
+
+
+
+UCLASS(BlueprintType)
+class MESHDESCRIPTION_API UMeshDescriptionBase : public UObject
+{
+	GENERATED_BODY()
+
+public:
+	// UObject interface
+	virtual void Serialize(FArchive& Ar) override;
+	//
+
+	/** Get a reference to the actual mesh description */
+	FMeshDescription& GetMeshDescription() { return MeshDescription; }
+	const FMeshDescription& GetMeshDescription() const { return MeshDescription; }
+
+	/** Set the mesh description */
+	void SetMeshDescription(FMeshDescription InMeshDescription)
+	{ 
+		MeshDescription = MoveTemp(InMeshDescription);
+	}
+
+	// UMeshDescriptionBase interface
+	virtual void RegisterAttributes();
+	virtual FMeshAttributes& GetRequiredAttributes() { return *RequiredAttributes; }
+	virtual const FMeshAttributes& GetRequiredAttributes() const { return *RequiredAttributes; }
+	//
+
+	/** Reset the contained mesh description */
+	void Reset();
+
+	/** Accessors for mesh element arrays */
+	FVertexArray& Vertices() { return MeshDescription.Vertices(); }
+	const FVertexArray& Vertices() const { return MeshDescription.Vertices(); }
+
+	FVertexInstanceArray& VertexInstances() { return MeshDescription.VertexInstances(); }
+	const FVertexInstanceArray& VertexInstances() const { return MeshDescription.VertexInstances(); }
+
+	FEdgeArray& Edges() { return MeshDescription.Edges(); }
+	const FEdgeArray& Edges() const { return MeshDescription.Edges(); }
+
+	FTriangleArray& Triangles() { return MeshDescription.Triangles(); }
+	const FTriangleArray& Triangles() const { return MeshDescription.Triangles(); }
+
+	FPolygonArray& Polygons() { return MeshDescription.Polygons(); }
+	const FPolygonArray& Polygons() const { return MeshDescription.Polygons(); }
+
+	FPolygonGroupArray& PolygonGroups() { return MeshDescription.PolygonGroups(); }
+	const FPolygonGroupArray& PolygonGroups() const { return MeshDescription.PolygonGroups(); }
+
+	/** Accessors for mesh element attributes */
+	TAttributesSet<FVertexID>& VertexAttributes() { return MeshDescription.VertexAttributes(); }
+	const TAttributesSet<FVertexID>& VertexAttributes() const { return MeshDescription.VertexAttributes(); }
+
+	TAttributesSet<FVertexInstanceID>& VertexInstanceAttributes() { return MeshDescription.VertexInstanceAttributes(); }
+	const TAttributesSet<FVertexInstanceID>& VertexInstanceAttributes() const { return MeshDescription.VertexInstanceAttributes(); }
+
+	TAttributesSet<FEdgeID>& EdgeAttributes() { return MeshDescription.EdgeAttributes(); }
+	const TAttributesSet<FEdgeID>& EdgeAttributes() const { return MeshDescription.EdgeAttributes(); }
+
+	TAttributesSet<FTriangleID>& TriangleAttributes() { return MeshDescription.TriangleAttributes(); }
+	const TAttributesSet<FTriangleID>& TriangleAttributes() const { return MeshDescription.TriangleAttributes(); }
+
+	TAttributesSet<FPolygonID>& PolygonAttributes() { return MeshDescription.PolygonAttributes(); }
+	const TAttributesSet<FPolygonID>& PolygonAttributes() const { return MeshDescription.PolygonAttributes(); }
+
+	TAttributesSet<FPolygonGroupID>& PolygonGroupAttributes() { return MeshDescription.PolygonGroupAttributes(); }
+	const TAttributesSet<FPolygonGroupID>& PolygonGroupAttributes() const { return MeshDescription.PolygonGroupAttributes(); }
+
+	/** Accessors for cached vertex position array */
+	TVertexAttributesRef<FVector> GetVertexPositions() { return GetRequiredAttributes().GetVertexPositions(); }
+	TVertexAttributesConstRef<FVector> GetVertexPositions() const { return GetRequiredAttributes().GetVertexPositions(); }
+
+public:
+
+	/** Empty the mesh description */
+	UFUNCTION(BlueprintCallable, Category="MeshDescription")
+	void Empty();
+
+	/** Return whether the mesh description is empty */
+	UFUNCTION(BlueprintPure, Category="MeshDescription")
+	bool IsEmpty() const;
+
+
+	///////////////////////////////////////////////////////
+	// Create / remove mesh elements
+
+	/** Reserves space for this number of new vertices */
+	UFUNCTION(BlueprintCallable, Category="MeshDescription")
+	void ReserveNewVertices(int32 NumberOfNewVertices);
+
+	/** Adds a new vertex to the mesh and returns its ID */
+	UFUNCTION(BlueprintCallable, Category="MeshDescription")
+	FVertexID CreateVertex();
+
+	/** Adds a new vertex to the mesh with the given ID */
+	UFUNCTION(BlueprintCallable, Category="MeshDescription")
+	void CreateVertexWithID(FVertexID VertexID);
+
+	/** Deletes a vertex from the mesh */
+	UFUNCTION(BlueprintCallable, Category="MeshDescription")
+	void DeleteVertex(FVertexID VertexID);
+
+	/** Returns whether the passed vertex ID is valid */
+	UFUNCTION(BlueprintPure, Category="MeshDescription")
+	bool IsVertexValid(FVertexID VertexID) const;
+
+	/** Reserves space for this number of new vertex instances */
+	UFUNCTION(BlueprintCallable, Category="MeshDescription")
+	void ReserveNewVertexInstances(int32 NumberOfNewVertexInstances);
+
+	/** Adds a new vertex instance to the mesh and returns its ID */
+	UFUNCTION(BlueprintCallable, Category="MeshDescription")
+	FVertexInstanceID CreateVertexInstance(FVertexID VertexID);
+
+	/** Adds a new vertex instance to the mesh with the given ID */
+	UFUNCTION(BlueprintCallable, Category="MeshDescription")
+	void CreateVertexInstanceWithID(FVertexInstanceID VertexInstanceID, FVertexID VertexID);
+
+	/** Deletes a vertex instance from a mesh */
+	UFUNCTION(BlueprintCallable, Category="MeshDescription")
+	void DeleteVertexInstance(FVertexInstanceID VertexInstanceID, TArray<FVertexID>& OrphanedVertices);
+
+	/** Returns whether the passed vertex instance ID is valid */
+	UFUNCTION(BlueprintPure, Category="MeshDescription")
+	bool IsVertexInstanceValid(FVertexInstanceID VertexInstanceID) const;
+
+	/** Reserves space for this number of new edges */
+	UFUNCTION(BlueprintCallable, Category="MeshDescription")
+	void ReserveNewEdges(int32 NumberOfNewEdges);
+
+	/** Adds a new edge to the mesh and returns its ID */
+	UFUNCTION(BlueprintCallable, Category="MeshDescription")
+	FEdgeID CreateEdge(FVertexID VertexID0, FVertexID VertexID1);
+
+	/** Adds a new edge to the mesh with the given ID */
+	UFUNCTION(BlueprintCallable, Category="MeshDescription")
+	void CreateEdgeWithID(FEdgeID EdgeID, FVertexID VertexID0, FVertexID VertexID1);
+
+	/** Deletes an edge from a mesh */
+	UFUNCTION(BlueprintCallable, Category="MeshDescription")
+	void DeleteEdge(FEdgeID EdgeID, TArray<FVertexID>& OrphanedVertices);
+
+	/** Returns whether the passed edge ID is valid */
+	UFUNCTION(BlueprintPure, Category="MeshDescription")
+	bool IsEdgeValid(FEdgeID EdgeID) const;
+
+	/** Reserves space for this number of new triangles */
+	UFUNCTION(BlueprintCallable, Category="MeshDescription")
+	void ReserveNewTriangles(int32 NumberOfNewTriangles);
+
+	/** Adds a new triangle to the mesh and returns its ID. This will also make an encapsulating polygon, and any missing edges. */
+	UFUNCTION(BlueprintCallable, Category="MeshDescription")
+	FTriangleID CreateTriangle(FPolygonGroupID PolygonGroupID, const TArray<FVertexInstanceID>& VertexInstanceIDs, TArray<FEdgeID>& NewEdgeIDs);
+
+	/** Adds a new triangle to the mesh with the given ID. This will also make an encapsulating polygon, and any missing edges. */
+	UFUNCTION(BlueprintCallable, Category="MeshDescription")
+	void CreateTriangleWithID(FTriangleID TriangleID, FPolygonGroupID PolygonGroupID, const TArray<FVertexInstanceID>& VertexInstanceIDs, TArray<FEdgeID>& NewEdgeIDs);
+
+	/** Deletes a triangle from the mesh */
+	UFUNCTION(BlueprintCallable, Category="MeshDescription")
+	void DeleteTriangle(FTriangleID TriangleID, TArray<FEdgeID>& OrphanedEdges, TArray<FVertexInstanceID>& OrphanedVertexInstances, TArray<FPolygonGroupID>& OrphanedPolygonGroupsPtr);
+
+	/** Returns whether the passed triangle ID is valid */
+	UFUNCTION(BlueprintPure, Category="MeshDescription")
+	bool IsTriangleValid(const FTriangleID TriangleID) const;
+
+	/** Reserves space for this number of new polygons */
+	UFUNCTION(BlueprintCallable, Category="MeshDescription")
+	void ReserveNewPolygons(const int32 NumberOfNewPolygons);
+
+	/** Adds a new polygon to the mesh and returns its ID. This will also make any missing edges, and all constituent triangles. */
+	UFUNCTION(BlueprintCallable, Category="MeshDescription")
+	FPolygonID CreatePolygon(FPolygonGroupID PolygonGroupID, TArray<FVertexInstanceID>& VertexInstanceIDs, TArray<FEdgeID>& NewEdgeIDs);
+
+	/** Adds a new polygon to the mesh with the given ID. This will also make any missing edges, and all constituent triangles. */
+	UFUNCTION(BlueprintCallable, Category="MeshDescription")
+	void CreatePolygonWithID(FPolygonID PolygonID, FPolygonGroupID PolygonGroupID, TArray<FVertexInstanceID>& VertexInstanceIDs, TArray<FEdgeID>& NewEdgeIDs);
+
+	/** Deletes a polygon from the mesh */
+	UFUNCTION(BlueprintCallable, Category="MeshDescription")
+	void DeletePolygon(FPolygonID PolygonID, TArray<FEdgeID>& OrphanedEdges, TArray<FVertexInstanceID>& OrphanedVertexInstances, TArray<FPolygonGroupID>& OrphanedPolygonGroups);
+
+	/** Returns whether the passed polygon ID is valid */
+	UFUNCTION(BlueprintPure, Category="MeshDescription")
+	bool IsPolygonValid(FPolygonID PolygonID) const;
+
+	/** Reserves space for this number of new polygon groups */
+	UFUNCTION(BlueprintCallable, Category="MeshDescription")
+	void ReserveNewPolygonGroups(int32 NumberOfNewPolygonGroups);
+
+	/** Adds a new polygon group to the mesh and returns its ID */
+	UFUNCTION(BlueprintCallable, Category="MeshDescription")
+	FPolygonGroupID CreatePolygonGroup();
+
+	/** Adds a new polygon group to the mesh with the given ID */
+	UFUNCTION(BlueprintCallable, Category="MeshDescription")
+	void CreatePolygonGroupWithID(FPolygonGroupID PolygonGroupID);
+
+	/** Deletes a polygon group from the mesh */
+	UFUNCTION(BlueprintCallable, Category="MeshDescription")
+	void DeletePolygonGroup(FPolygonGroupID PolygonGroupID);
+
+	/** Returns whether the passed polygon group ID is valid */
+	UFUNCTION(BlueprintPure, Category="MeshDescription")
+	bool IsPolygonGroupValid(FPolygonGroupID PolygonGroupID) const;
+
+
+	//////////////////////////////////////////////////////////////////////
+	// Vertex operations
+
+	/** Returns whether a given vertex is orphaned, i.e. it doesn't form part of any polygon */
+	UFUNCTION(BlueprintPure, Category="MeshDescription")
+	bool IsVertexOrphaned(FVertexID VertexID) const;
+
+	/** Returns the edge ID defined by the two given vertex IDs, if there is one; otherwise FEdgeID::Invalid */
+	UFUNCTION(BlueprintPure, Category="MeshDescription")
+	FEdgeID GetVertexPairEdge(FVertexID VertexID0, FVertexID VertexID1) const;
+
+	/** Returns reference to an array of Edge IDs connected to this vertex */
+	UFUNCTION(BlueprintPure, Category="MeshDescription")
+	void GetVertexConnectedEdges(FVertexID VertexID, TArray<FEdgeID>& OutEdgeIDs) const;
+
+	/** Returns number of edges connected to this vertex */
+	UFUNCTION(BlueprintPure, Category="MeshDescription")
+	int32 GetNumVertexConnectedEdges(FVertexID VertexID) const;
+
+	/** Returns reference to an array of VertexInstance IDs instanced from this vertex */
+	UFUNCTION(BlueprintPure, Category="MeshDescription")
+	void GetVertexVertexInstances(FVertexID VertexID, TArray<FVertexInstanceID>& OutVertexInstanceIDs) const;
+
+	/** Returns number of vertex instances created from this vertex */
+	UFUNCTION(BlueprintPure, Category="MeshDescription")
+	int32 GetNumVertexVertexInstances(FVertexID VertexID) const;
+
+	/** Returns the triangles connected to this vertex */
+	UFUNCTION(BlueprintPure, Category="MeshDescription")
+	void GetVertexConnectedTriangles(FVertexID VertexID, TArray<FTriangleID>& OutConnectedTriangleIDs) const;
+
+	/** Returns number of triangles connected to this vertex */
+	UFUNCTION(BlueprintPure, Category="MeshDescription")
+	int32 GetNumVertexConnectedTriangles(FVertexID VertexID) const;
+
+	/** Returns the polygons connected to this vertex */
+	UFUNCTION(BlueprintPure, Category="MeshDescription")
+	void GetVertexConnectedPolygons(FVertexID VertexID, TArray<FPolygonID>& OutConnectedPolygonIDs) const;
+
+	/** Returns the number of polygons connected to this vertex */
+	UFUNCTION(BlueprintPure, Category="MeshDescription")
+	int32 GetNumVertexConnectedPolygons(FVertexID VertexID) const;
+
+	/** Returns the vertices adjacent to this vertex */
+	UFUNCTION(BlueprintPure, Category="MeshDescription")
+	void GetVertexAdjacentVertices(FVertexID VertexID, TArray<FVertexID>& OutAdjacentVertexIDs) const;
+
+	/** Gets a vertex position */
+	UFUNCTION(BlueprintPure, Category="MeshDescription")
+	FVector GetVertexPosition(FVertexID VertexID) const;
+
+	/** Sets a vertex position */
+	UFUNCTION(BlueprintCallable, Category="MeshDescription")
+	void SetVertexPosition(FVertexID VertexID, const FVector& Position);
+
+
+	//////////////////////////////////////////////////////////////////////
+	// Vertex instance operations
+
+	/** Returns the vertex ID associated with the given vertex instance */
+	UFUNCTION(BlueprintPure, Category="MeshDescription")
+	FVertexID GetVertexInstanceVertex(FVertexInstanceID VertexInstanceID) const;
+
+	/** Returns the edge ID defined by the two given vertex instance IDs, if there is one; otherwise FEdgeID::Invalid */
+	UFUNCTION(BlueprintPure, Category="MeshDescription")
+	FEdgeID GetVertexInstancePairEdge(FVertexInstanceID VertexInstanceID0, FVertexInstanceID VertexInstanceID1) const;
+
+	/** Returns reference to an array of Triangle IDs connected to this vertex instance */
+	UFUNCTION(BlueprintPure, Category="MeshDescription")
+	void GetVertexInstanceConnectedTriangles(FVertexInstanceID VertexInstanceID, TArray<FTriangleID>& OutConnectedTriangleIDs) const;
+
+	/** Returns the number of triangles connected to this vertex instance */
+	UFUNCTION(BlueprintPure, Category="MeshDescription")
+	int32 GetNumVertexInstanceConnectedTriangles(FVertexInstanceID VertexInstanceID) const;
+
+	/** Returns the polygons connected to this vertex instance */
+	UFUNCTION(BlueprintPure, Category="MeshDescription")
+	void GetVertexInstanceConnectedPolygons(FVertexInstanceID VertexInstanceID, TArray<FPolygonID>& OutConnectedPolygonIDs) const;
+
+	/** Returns the number of polygons connected to this vertex instance. */
+	UFUNCTION(BlueprintPure, Category="MeshDescription")
+	int32 GetNumVertexInstanceConnectedPolygons(FVertexInstanceID VertexInstanceID) const;
+
+
+	//////////////////////////////////////////////////////////////////////
+	// Edge operations
+
+	/** Determine whether a given edge is an internal edge between triangles of a polygon */
+	UFUNCTION(BlueprintPure, Category="MeshDescription")
+	bool IsEdgeInternal(FEdgeID EdgeID) const;
+
+	/** Determine whether a given edge is an internal edge between triangles of a specific polygon */
+	UFUNCTION(BlueprintPure, Category="MeshDescription")
+	bool IsEdgeInternalToPolygon(FEdgeID EdgeID, FPolygonID PolygonID) const;
+
+	/** Returns reference to an array of triangle IDs connected to this edge */
+	UFUNCTION(BlueprintPure, Category="MeshDescription")
+	void GetEdgeConnectedTriangles(FEdgeID EdgeID, TArray<FTriangleID>& OutConnectedTriangleIDs) const;
+
+	/** Returns the number of triangles connected to this edge */
+	UFUNCTION(BlueprintPure, Category="MeshDescription")
+	int32 GetNumEdgeConnectedTriangles(FEdgeID EdgeID) const;
+
+	/** Returns the polygons connected to this edge */
+	UFUNCTION(BlueprintPure, Category="MeshDescription")
+	void GetEdgeConnectedPolygons(FEdgeID EdgeID, TArray<FPolygonID>& OutConnectedPolygonIDs) const;
+
+	/** Returns the number of polygons connected to this edge */
+	UFUNCTION(BlueprintPure, Category="MeshDescription")
+	int32 GetNumEdgeConnectedPolygons(FEdgeID EdgeID) const;
+
+	/** Returns the vertex ID corresponding to one of the edge endpoints */
+	UFUNCTION(BlueprintPure, Category="MeshDescription")
+	FVertexID GetEdgeVertex(FEdgeID EdgeID, int32 VertexNumber) const;
+
+	/** Returns a pair of vertex IDs defining the edge */
+	UFUNCTION(BlueprintPure, Category="MeshDescription")
+	void GetEdgeVertices(const FEdgeID EdgeID, TArray<FVertexID>& OutVertexIDs) const;
+
+
+	//////////////////////////////////////////////////////////////////////
+	// Triangle operations
+
+	/** Get the polygon which contains this triangle */
+	UFUNCTION(BlueprintPure, Category="MeshDescription")
+	FPolygonID GetTrianglePolygon(FTriangleID TriangleID) const;
+
+	/** Get the polygon group which contains this triangle */
+	UFUNCTION(BlueprintPure, Category="MeshDescription")
+	FPolygonGroupID GetTrianglePolygonGroup(FTriangleID TriangleID) const;
+
+	/** Determines if this triangle is part of an n-gon */
+	UFUNCTION(BlueprintPure, Category="MeshDescription")
+	bool IsTrianglePartOfNgon(FTriangleID TriangleID) const;
+
+	/** Get the vertex instances which define this triangle */
+	UFUNCTION(BlueprintPure, Category="MeshDescription")
+	void GetTriangleVertexInstances(FTriangleID TriangleID, TArray<FVertexInstanceID>& OutVertexInstanceIDs) const;
+
+	/** Get the specified vertex instance by index */
+	UFUNCTION(BlueprintPure, Category="MeshDescription")
+	FVertexInstanceID GetTriangleVertexInstance(FTriangleID TriangleID, int32 Index) const;
+
+	/** Returns the vertices which define this triangle */
+	UFUNCTION(BlueprintPure, Category="MeshDescription")
+	void GetTriangleVertices(FTriangleID TriangleID, TArray<FVertexID>& OutVertexIDs) const;
+
+	/** Returns the edges which define this triangle */
+	UFUNCTION(BlueprintPure, Category="MeshDescription")
+	void GetTriangleEdges(FTriangleID TriangleID, TArray<FEdgeID>& OutEdgeIDs) const;
+
+	/** Returns the adjacent triangles to this triangle */
+	UFUNCTION(BlueprintPure, Category="MeshDescription")
+	void GetTriangleAdjacentTriangles(FTriangleID TriangleID, TArray<FTriangleID>& OutTriangleIDs) const;
+
+	/** Return the vertex instance which corresponds to the given vertex on the given triangle, or FVertexInstanceID::Invalid */
+	UFUNCTION(BlueprintPure, Category="MeshDescription")
+	FVertexInstanceID GetVertexInstanceForTriangleVertex(FTriangleID TriangleID, FVertexID VertexID) const;
+
+
+	//////////////////////////////////////////////////////////////////////
+	// Polygon operations
+
+	/** Return reference to an array of triangle IDs which comprise this polygon */
+	UFUNCTION(BlueprintPure, Category="MeshDescription")
+	void GetPolygonTriangles(FPolygonID PolygonID, TArray<FTriangleID>& OutTriangleIDs) const;
+
+	/** Return the number of triangles which comprise this polygon */
+	UFUNCTION(BlueprintPure, Category="MeshDescription")
+	int32 GetNumPolygonTriangles(FPolygonID PolygonID) const;
+
+	/** Returns reference to an array of VertexInstance IDs forming the perimeter of this polygon */
+	UFUNCTION(BlueprintPure, Category="MeshDescription")
+	void GetPolygonVertexInstances(FPolygonID PolygonID, TArray<FVertexInstanceID>& OutVertexInstanceIDs) const;
+
+	/** Returns the number of vertices this polygon has */
+	UFUNCTION(BlueprintPure, Category="MeshDescription")
+	int32 GetNumPolygonVertices(FPolygonID PolygonID) const;
+
+	/** Returns the vertices which form the polygon perimeter */
+	UFUNCTION(BlueprintPure, Category="MeshDescription")
+	void GetPolygonVertices(FPolygonID PolygonID, TArray<FVertexID>& OutVertexIDs) const;
+
+	/** Returns the edges which form the polygon perimeter */
+	UFUNCTION(BlueprintPure, Category="MeshDescription")
+	void GetPolygonPerimeterEdges(FPolygonID PolygonID, TArray<FEdgeID>& OutEdgeIDs) const;
+
+	/** Populate the provided array with a list of edges which are internal to the polygon, i.e. those which separate
+	    constituent triangles. */
+	UFUNCTION(BlueprintPure, Category="MeshDescription")
+	void GetPolygonInternalEdges(FPolygonID PolygonID, TArray<FEdgeID>& OutEdgeIDs) const;
+
+	/** Return the number of internal edges in this polygon */
+	UFUNCTION(BlueprintPure, Category="MeshDescription")
+	int32 GetNumPolygonInternalEdges(FPolygonID PolygonID) const;
+
+	/** Populates the passed array with adjacent polygons */
+	UFUNCTION(BlueprintPure, Category="MeshDescription")
+	void GetPolygonAdjacentPolygons(FPolygonID PolygonID, TArray<FPolygonID>& OutPolygonIDs) const;
+
+	/** Return the polygon group associated with a polygon */
+	UFUNCTION(BlueprintPure, Category="MeshDescription")
+	FPolygonGroupID GetPolygonPolygonGroup(FPolygonID PolygonID) const;
+
+	/** Return the vertex instance which corresponds to the given vertex on the given polygon, or FVertexInstanceID::Invalid */
+	UFUNCTION(BlueprintPure, Category="MeshDescription")
+	FVertexInstanceID GetVertexInstanceForPolygonVertex(FPolygonID PolygonID, FVertexID VertexID) const;
+
+	/** Set the vertex instance at the given index around the polygon to the new value */
+	UFUNCTION(BlueprintCallable, Category="MeshDescription")
+	void SetPolygonVertexInstance(FPolygonID PolygonID, int32 PerimeterIndex, FVertexInstanceID VertexInstanceID);
+
+	/** Sets the polygon group associated with a polygon */
+	UFUNCTION(BlueprintCallable, Category="MeshDescription")
+	void SetPolygonPolygonGroup(FPolygonID PolygonID, FPolygonGroupID PolygonGroupID);
+
+	/** Reverse the winding order of the vertices of this polygon */
+	UFUNCTION(BlueprintCallable, Category="MeshDescription")
+	void ReversePolygonFacing(FPolygonID PolygonID);
+
+	/** Generates triangles and internal edges for the given polygon */
+	UFUNCTION(BlueprintCallable, Category="MeshDescription")
+	void ComputePolygonTriangulation(FPolygonID PolygonID);
+
+
+	//////////////////////////////////////////////////////////////////////
+	// Polygon group operations
+
+	/** Returns the polygons associated with the given polygon group */
+	UFUNCTION(BlueprintPure, Category="MeshDescription")
+	void GetPolygonGroupPolygons(FPolygonGroupID PolygonGroupID, TArray<FPolygonID>& OutPolygonIDs) const;
+
+	/** Returns the number of polygons in this polygon group */
+	UFUNCTION(BlueprintPure, Category="MeshDescription")
+	int32 GetNumPolygonGroupPolygons(FPolygonGroupID PolygonGroupID) const;
+
+protected:
+	FMeshDescription MeshDescription;
+	TUniquePtr<FMeshAttributes> RequiredAttributes;
+};
diff --git a/Engine/Source/Runtime/MovieSceneTracks/Public/Sections/MovieScene3DTransformSection.h b/Engine/Source/Runtime/MovieSceneTracks/Public/Sections/MovieScene3DTransformSection.h
--- a/Engine/Source/Runtime/MovieSceneTracks/Public/Sections/MovieScene3DTransformSection.h
+++ b/Engine/Source/Runtime/MovieSceneTracks/Public/Sections/MovieScene3DTransformSection.h
@@ -9,6 +9,7 @@
 #include "MovieSceneSection.h"
 #include "MovieSceneKeyStruct.h"
 #include "Channels/MovieSceneFloatChannel.h"
+#include "TransformData.h"
 #include "MovieScene3DTransformSection.generated.h"
 
 
@@ -24,36 +25,6 @@ enum class EShow3DTrajectory : uint8
 #endif
 
 
-/**
-* Stores information about a transform for the purpose of adding keys to a transform section
-*/
-struct FTransformData
-{
-	/** Translation component */
-	FVector Translation;
-	/** Rotation component */
-	FRotator Rotation;
-	/** Scale component */
-	FVector Scale;
-
-	FTransformData()
-		: Translation(ForceInitToZero)
-		, Rotation(ForceInitToZero)
-		, Scale(ForceInitToZero)
-	{}
-
-	/**
-	* Constructor.  Builds the data from a scene component
-	* Uses relative transform only
-	*
-	* @param InComponent	The component to build from
-	*/
-	FTransformData(const USceneComponent* InComponent)
-		: Translation(InComponent->GetRelativeLocation())
-		, Rotation(InComponent->GetRelativeRotation())
-		, Scale(InComponent->GetRelativeScale3D())
-	{}
-};
 
 /**
  * Proxy structure for translation keys in 3D transform sections.
diff --git a/Engine/Source/Runtime/StaticMeshDescription/Private/StaticMeshAttributes.cpp b/Engine/Source/Runtime/StaticMeshDescription/Private/StaticMeshAttributes.cpp
new file mode 100644
--- /dev/null
+++ b/Engine/Source/Runtime/StaticMeshDescription/Private/StaticMeshAttributes.cpp
@@ -0,0 +1,54 @@
+// Copyright 1998-2019 Epic Games, Inc. All Rights Reserved.
+
+#include "StaticMeshAttributes.h"
+
+namespace MeshAttribute
+{
+	const FName Vertex::CornerSharpness("CornerSharpness");
+
+	const FName VertexInstance::TextureCoordinate("TextureCoordinate");
+	const FName VertexInstance::Normal("Normal");
+	const FName VertexInstance::Tangent("Tangent");
+	const FName VertexInstance::BinormalSign("BinormalSign");
+	const FName VertexInstance::Color("Color");
+
+	const FName Edge::IsHard("IsHard");
+	const FName Edge::IsUVSeam("IsUVSeam");
+	const FName Edge::CreaseSharpness("CreaseSharpness");
+
+	const FName Polygon::Normal("Normal");
+	const FName Polygon::Tangent("Tangent");
+	const FName Polygon::Binormal("Binormal");
+	const FName Polygon::Center("Center");
+
+	const FName PolygonGroup::ImportedMaterialSlotName("ImportedMaterialSlotName");
+	const FName PolygonGroup::EnableCollision("EnableCollision");
+	const FName PolygonGroup::CastShadow("CastShadow");
+}
+
+
+
+void FStaticMeshAttributes::Register()
+{
+	// Add basic vertex attributes
+	MeshDescription.VertexAttributes().RegisterAttribute<float>(MeshAttribute::Vertex::CornerSharpness, 1, 0.0f, EMeshAttributeFlags::Lerpable);
+
+	// Add basic vertex instance attributes
+	MeshDescription.VertexInstanceAttributes().RegisterAttribute<FVector2D>(MeshAttribute::VertexInstance::TextureCoordinate, 1, FVector2D::ZeroVector, EMeshAttributeFlags::Lerpable);
+	MeshDescription.VertexInstanceAttributes().RegisterAttribute<FVector>(MeshAttribute::VertexInstance::Normal, 1, FVector::ZeroVector, EMeshAttributeFlags::AutoGenerated);
+	MeshDescription.VertexInstanceAttributes().RegisterAttribute<FVector>(MeshAttribute::VertexInstance::Tangent, 1, FVector::ZeroVector, EMeshAttributeFlags::AutoGenerated);
+	MeshDescription.VertexInstanceAttributes().RegisterAttribute<float>(MeshAttribute::VertexInstance::BinormalSign, 1, 0.0f, EMeshAttributeFlags::AutoGenerated);
+	MeshDescription.VertexInstanceAttributes().RegisterAttribute<FVector4>(MeshAttribute::VertexInstance::Color, 1, FVector4(1.0f, 1.0f, 1.0f, 1.0f), EMeshAttributeFlags::Lerpable);
+
+	// Add basic edge attributes
+	MeshDescription.EdgeAttributes().RegisterAttribute<bool>(MeshAttribute::Edge::IsHard, 1, false);
+	MeshDescription.EdgeAttributes().RegisterAttribute<float>(MeshAttribute::Edge::CreaseSharpness, 1, 0.0f, EMeshAttributeFlags::Lerpable);
+
+	// Add basic polygon attributes
+
+	// Add basic polygon group attributes
+	MeshDescription.PolygonGroupAttributes().RegisterAttribute<FName>(MeshAttribute::PolygonGroup::ImportedMaterialSlotName); //The unique key to match the mesh material slot
+
+	// Call super class
+	FMeshAttributes::Register();
+}
diff --git a/Engine/Source/Runtime/StaticMeshDescription/Private/StaticMeshDescription.cpp b/Engine/Source/Runtime/StaticMeshDescription/Private/StaticMeshDescription.cpp
new file mode 100644
--- /dev/null
+++ b/Engine/Source/Runtime/StaticMeshDescription/Private/StaticMeshDescription.cpp
@@ -0,0 +1,133 @@
+// Copyright 1998-2019 Epic Games, Inc. All Rights Reserved.
+
+#include "StaticMeshDescription.h"
+#include "StaticMeshAttributes.h"
+#include "StaticMeshOperations.h"
+
+
+void UStaticMeshDescription::RegisterAttributes()
+{
+	RequiredAttributes = MakeUnique<FStaticMeshAttributes>(MeshDescription);
+	RequiredAttributes->Register();
+}
+
+
+FVector2D UStaticMeshDescription::GetVertexInstanceUV(FVertexInstanceID VertexInstanceID, int32 UVIndex) const
+{
+	if (!MeshDescription.IsVertexInstanceValid(VertexInstanceID))
+	{
+		UE_LOG(LogMeshDescription, Warning, TEXT("GetVertexInstanceUV: VertexInstanceID %d doesn't exist."), VertexInstanceID.GetValue());
+		return FVector2D::ZeroVector;
+	}
+
+	if (!MeshDescription.VertexInstanceAttributes().HasAttribute(MeshAttribute::VertexInstance::TextureCoordinate))
+	{
+		UE_LOG(LogMeshDescription, Warning, TEXT("GetVertexInstanceUV: VertexInstanceAttribute TextureCoordinate doesn't exist."));
+		return FVector2D::ZeroVector;
+	}
+
+	return MeshDescription.VertexInstanceAttributes().GetAttribute<FVector2D>(VertexInstanceID, MeshAttribute::VertexInstance::TextureCoordinate, UVIndex);
+}
+
+
+void UStaticMeshDescription::SetVertexInstanceUV(FVertexInstanceID VertexInstanceID, FVector2D UV, int32 UVIndex)
+{
+	if (!MeshDescription.IsVertexInstanceValid(VertexInstanceID))
+	{
+		UE_LOG(LogMeshDescription, Warning, TEXT("SetVertexInstanceUV: VertexInstanceID %d doesn't exist."), VertexInstanceID.GetValue());
+		return;
+	}
+
+	if (!MeshDescription.VertexInstanceAttributes().HasAttribute(MeshAttribute::VertexInstance::TextureCoordinate))
+	{
+		UE_LOG(LogMeshDescription, Warning, TEXT("SetVertexInstanceUV: VertexInstanceAttribute TextureCoordinate doesn't exist."));
+		return;
+	}
+
+	MeshDescription.VertexInstanceAttributes().SetAttribute(VertexInstanceID, MeshAttribute::VertexInstance::TextureCoordinate, UVIndex, UV);
+}
+
+
+void UStaticMeshDescription::SetPolygonGroupMaterialSlotName(FPolygonGroupID PolygonGroupID, const FName& SlotName)
+{
+	if (!MeshDescription.IsPolygonGroupValid(PolygonGroupID))
+	{
+		UE_LOG(LogMeshDescription, Warning, TEXT("SetPolygonGroupMaterialSlotName: PolygonGroupID %d doesn't exist."), PolygonGroupID.GetValue());
+		return;
+	}
+
+	if (!MeshDescription.PolygonGroupAttributes().HasAttribute(MeshAttribute::PolygonGroup::ImportedMaterialSlotName))
+	{
+		UE_LOG(LogMeshDescription, Warning, TEXT("SetPolygonGroupMaterialSlotName: PolygonGroupAttribute ImportedMaterialSlotName doesn't exist."));
+		return;
+	}
+
+	MeshDescription.PolygonGroupAttributes().SetAttribute(PolygonGroupID, MeshAttribute::PolygonGroup::ImportedMaterialSlotName, 0, SlotName);
+}
+
+
+void UStaticMeshDescription::CreateCube(FVector Center, FVector HalfExtents, FPolygonGroupID PolygonGroup,
+										FPolygonID& PolygonID_PlusX,
+										FPolygonID& PolygonID_MinusX,
+										FPolygonID& PolygonID_PlusY, 
+										FPolygonID& PolygonID_MinusY,
+										FPolygonID& PolygonID_PlusZ,
+										FPolygonID& PolygonID_MinusZ)
+{
+	TVertexAttributesRef<FVector> Positions = GetVertexPositions();
+
+	FVertexID VertexIDs[8];
+
+	MeshDescription.ReserveNewVertices(8);
+	for (int32 Index = 0; Index < 8; ++Index)
+	{
+		VertexIDs[Index] = MeshDescription.CreateVertex();
+	}
+
+	Positions[VertexIDs[0]] = Center + HalfExtents * FVector( 1.0f, -1.0f,  1.0f);
+	Positions[VertexIDs[1]] = Center + HalfExtents * FVector( 1.0f,  1.0f,  1.0f);
+	Positions[VertexIDs[2]] = Center + HalfExtents * FVector(-1.0f,  1.0f,  1.0f);
+	Positions[VertexIDs[3]] = Center + HalfExtents * FVector(-1.0f, -1.0f,  1.0f);
+	Positions[VertexIDs[4]] = Center + HalfExtents * FVector(-1.0f,  1.0f, -1.0f);
+	Positions[VertexIDs[5]] = Center + HalfExtents * FVector(-1.0f, -1.0f, -1.0f);
+	Positions[VertexIDs[6]] = Center + HalfExtents * FVector( 1.0f, -1.0f, -1.0f);
+	Positions[VertexIDs[7]] = Center + HalfExtents * FVector( 1.0f,  1.0f, -1.0f);
+
+	auto MakePolygon = [this, &VertexIDs, PolygonGroup](int32 P0, int32 P1, int32 P2, int32 P3) -> FPolygonID
+	{
+		FVertexInstanceID VertexInstanceIDs[4];
+		VertexInstanceIDs[0] = MeshDescription.CreateVertexInstance(VertexIDs[P0]);
+		VertexInstanceIDs[1] = MeshDescription.CreateVertexInstance(VertexIDs[P1]);
+		VertexInstanceIDs[2] = MeshDescription.CreateVertexInstance(VertexIDs[P2]);
+		VertexInstanceIDs[3] = MeshDescription.CreateVertexInstance(VertexIDs[P3]);
+
+		TArray<FEdgeID> EdgeIDs;
+		EdgeIDs.Reserve(4);
+
+		FPolygonID PolygonID = MeshDescription.CreatePolygon(PolygonGroup, VertexInstanceIDs, &EdgeIDs);
+
+		for (FEdgeID EdgeID : EdgeIDs)
+		{
+			GetEdgeHardnesses()[EdgeID] = true;
+		}
+
+		return PolygonID;
+	};
+
+	PolygonID_PlusX = MakePolygon(0, 1, 7, 6);
+	PolygonID_MinusX = MakePolygon(2, 3, 5, 4);
+	PolygonID_PlusY = MakePolygon(1, 2, 4, 7);
+	PolygonID_MinusY = MakePolygon(3, 0, 6, 5);
+	PolygonID_PlusZ = MakePolygon(1, 0, 3, 2);
+	PolygonID_MinusZ = MakePolygon(5, 4, 7, 6);
+
+	MeshDescription.PolygonAttributes().RegisterAttribute<FVector>(MeshAttribute::Polygon::Normal, 1, FVector::ZeroVector, EMeshAttributeFlags::Transient);
+	MeshDescription.PolygonAttributes().RegisterAttribute<FVector>(MeshAttribute::Polygon::Tangent, 1, FVector::ZeroVector, EMeshAttributeFlags::Transient);
+	MeshDescription.PolygonAttributes().RegisterAttribute<FVector>(MeshAttribute::Polygon::Binormal, 1, FVector::ZeroVector, EMeshAttributeFlags::Transient);
+	MeshDescription.PolygonAttributes().RegisterAttribute<FVector>(MeshAttribute::Polygon::Center, 1, FVector::ZeroVector, EMeshAttributeFlags::Transient);
+
+	FStaticMeshOperations::ComputePolygonTangentsAndNormals(MeshDescription);
+	FStaticMeshOperations::ComputeTangentsAndNormals(MeshDescription, EComputeNTBsFlags::Normals | EComputeNTBsFlags::Tangents);
+}
+
+
