Corrected deprecation version.

diff --git a/Engine/Source/Runtime/MeshDescription/Public/MeshDescription.h b/Engine/Source/Runtime/MeshDescription/Public/MeshDescription.h
--- a/Engine/Source/Runtime/MeshDescription/Public/MeshDescription.h
+++ b/Engine/Source/Runtime/MeshDescription/Public/MeshDescription.h
@@ -1294,35 +1294,35 @@ public:
 	void TriangulateMesh();
 
 	/** Set the polygon tangent and normal only for the specified polygonIDs */
-	UE_DEPRECATED(4.25, "Please use FStaticMeshOperations::ComputePolygonTangentsAndNormals.")
+	UE_DEPRECATED(4.24, "Please use FStaticMeshOperations::ComputePolygonTangentsAndNormals.")
 	void ComputePolygonTangentsAndNormals(const TArray<FPolygonID>& PolygonIDs, float ComparisonThreshold = 0.0f);
 
 	/** Set the polygon tangent and normal for all polygons in the mesh description. */
-	UE_DEPRECATED(4.25, "Please use FStaticMeshOperations::ComputePolygonTangentsAndNormals.")
+	UE_DEPRECATED(4.24, "Please use FStaticMeshOperations::ComputePolygonTangentsAndNormals.")
 	void ComputePolygonTangentsAndNormals(float ComparisonThreshold = 0.0f);
 	
 	/** Set the vertex instance tangent and normal only for the specified VertexInstanceIDs */
-	UE_DEPRECATED(4.25, "Please use FStaticMeshOperations::ComputeTangentsAndNormals.")
+	UE_DEPRECATED(4.24, "Please use FStaticMeshOperations::ComputeTangentsAndNormals.")
 	void ComputeTangentsAndNormals(const TArray<FVertexInstanceID>& VertexInstanceIDs, EComputeNTBsOptions ComputeNTBsOptions);
 
 	/** Set the vertex instance tangent and normal for all vertex instances in the mesh description. */
-	UE_DEPRECATED(4.25, "Please use FStaticMeshOperations::ComputeTangentsAndNormals.")
+	UE_DEPRECATED(4.24, "Please use FStaticMeshOperations::ComputeTangentsAndNormals.")
 	void ComputeTangentsAndNormals(EComputeNTBsOptions ComputeNTBsOptions);
 
 	/** Determine the edge hardnesses from existing normals */
-	UE_DEPRECATED(4.25, "Please use FStaticMeshOperations::DetermineEdgeHardnessesFromVertexInstanceNormals.")
+	UE_DEPRECATED(4.24, "Please use FStaticMeshOperations::DetermineEdgeHardnessesFromVertexInstanceNormals.")
 	void DetermineEdgeHardnessesFromVertexInstanceNormals( const float Tolerance = KINDA_SMALL_NUMBER );
 
 	/** Determine UV seams from existing vertex instance UVs */
-	UE_DEPRECATED(4.25, "UVSeam attribute is now deprecated in MeshDescription.")
+	UE_DEPRECATED(4.24, "UVSeam attribute is now deprecated in MeshDescription.")
 	void DetermineUVSeamsFromUVs( const int32 UVIndex, const float Tolerance = KINDA_SMALL_NUMBER );
 
 	/** Get polygons in the same UV chart as the specified polygon */
-	UE_DEPRECATED(4.25, "UVSeam attribute is now deprecated in MeshDescription.")
+	UE_DEPRECATED(4.24, "UVSeam attribute is now deprecated in MeshDescription.")
 	void GetPolygonsInSameChartAsPolygon( const FPolygonID PolygonID, TArray<FPolygonID>& OutPolygonIDs );
 
 	/** Get array of all UV charts */
-	UE_DEPRECATED(4.25, "UVSeam attribute is now deprecated in MeshDescription.")
+	UE_DEPRECATED(4.24, "UVSeam attribute is now deprecated in MeshDescription.")
 	void GetAllCharts( TArray<TArray<FPolygonID>>& OutCharts );
 
 	/** Reverses the winding order of all polygons in the mesh */
