MeshDescription now uses triangles as its primary face representation. They can and should be used in preference to polygons when a mesh is known to be entirely triangulated, as they are quicker, and 

diff --git a/Engine/Source/Runtime/MeshDescription/Private/MeshDescription.cpp b/Engine/Source/Runtime/MeshDescription/Private/MeshDescription.cpp
--- a/Engine/Source/Runtime/MeshDescription/Private/MeshDescription.cpp
+++ b/Engine/Source/Runtime/MeshDescription/Private/MeshDescription.cpp
@@ -4,6 +4,7 @@
 #include "MeshAttributes.h"
 #include "Serialization/BulkDataReader.h"
 #include "Serialization/BulkDataWriter.h"
+#include "Algo/Copy.h"
 
 
 void UDEPRECATED_MeshDescription::Serialize( FArchive& Ar )
@@ -20,75 +21,156 @@ void UDEPRECATED_MeshDescription::Serialize( FArchive& Ar )
 }
 
 
-FArchive& operator<<( FArchive& Ar, FMeshDescription& MeshDescription )
+PRAGMA_DISABLE_DEPRECATION_WARNINGS
+
+void FMeshDescription::Serialize(FArchive& Ar)
 {
-	Ar.UsingCustomVersion( FReleaseObjectVersion::GUID );
-	Ar.UsingCustomVersion( FEditorObjectVersion::GUID );
+	Ar.UsingCustomVersion(FReleaseObjectVersion::GUID);
+	Ar.UsingCustomVersion(FEditorObjectVersion::GUID);
 
-	if( Ar.IsLoading() && Ar.CustomVer( FReleaseObjectVersion::GUID ) < FReleaseObjectVersion::MeshDescriptionNewSerialization )
+	if (Ar.IsLoading() && Ar.CustomVer(FReleaseObjectVersion::GUID) < FReleaseObjectVersion::MeshDescriptionNewSerialization)
 	{
-		UE_LOG( LogLoad, Warning, TEXT( "Deprecated serialization format" ) );
+		UE_LOG(LogLoad, Warning, TEXT("Deprecated serialization format"));
 	}
 
-	Ar << MeshDescription.VertexArray;
-	Ar << MeshDescription.VertexInstanceArray;
-	Ar << MeshDescription.EdgeArray;
-	Ar << MeshDescription.PolygonArray;
-	Ar << MeshDescription.PolygonGroupArray;
+	Ar << VertexArray;
+	Ar << VertexInstanceArray;
+	Ar << EdgeArray;
+	Ar << PolygonArray;
+	Ar << PolygonGroupArray;
 
-	Ar << MeshDescription.VertexAttributesSet;
-	Ar << MeshDescription.VertexInstanceAttributesSet;
-	Ar << MeshDescription.EdgeAttributesSet;
-	Ar << MeshDescription.PolygonAttributesSet;
-	Ar << MeshDescription.PolygonGroupAttributesSet;
+	Ar << VertexAttributesSet;
+	Ar << VertexInstanceAttributesSet;
+	Ar << EdgeAttributesSet;
+	Ar << PolygonAttributesSet;
+	Ar << PolygonGroupAttributesSet;
 
-	if( Ar.IsLoading() && Ar.CustomVer( FReleaseObjectVersion::GUID ) >= FReleaseObjectVersion::MeshDescriptionNewSerialization )
+	// Serialize new triangle arrays since version MeshDescriptionTriangles
+	if (!Ar.IsLoading() || Ar.CustomVer(FEditorObjectVersion::GUID) >= FEditorObjectVersion::MeshDescriptionTriangles)
+	{
+		Ar << TriangleArray;
+		Ar << TriangleAttributesSet;
+	}
+
+	if (Ar.IsLoading() && Ar.CustomVer(FReleaseObjectVersion::GUID) >= FReleaseObjectVersion::MeshDescriptionNewSerialization)
 	{
 		// Populate vertex instance IDs for vertices
-		for( const FVertexInstanceID VertexInstanceID : MeshDescription.VertexInstanceArray.GetElementIDs() )
+		for (const FVertexInstanceID VertexInstanceID : VertexInstanceArray.GetElementIDs())
 		{
-			const FVertexID VertexID = MeshDescription.GetVertexInstanceVertex( VertexInstanceID );
-			MeshDescription.VertexArray[ VertexID ].VertexInstanceIDs.Add( VertexInstanceID );
+			const FVertexID VertexID = GetVertexInstanceVertex(VertexInstanceID);
+			VertexArray[VertexID].VertexInstanceIDs.Add(VertexInstanceID);
 		}
 
 		// Populate edge IDs for vertices
-		for( const FEdgeID EdgeID : MeshDescription.EdgeArray.GetElementIDs() )
+		for (const FEdgeID EdgeID : EdgeArray.GetElementIDs())
 		{
-			const FVertexID VertexID0 = MeshDescription.GetEdgeVertex( EdgeID, 0 );
-			const FVertexID VertexID1 = MeshDescription.GetEdgeVertex( EdgeID, 1 );
-			MeshDescription.VertexArray[ VertexID0 ].ConnectedEdgeIDs.Add( EdgeID );
-			MeshDescription.VertexArray[ VertexID1 ].ConnectedEdgeIDs.Add( EdgeID );
+			const FVertexID VertexID0 = GetEdgeVertex(EdgeID, 0);
+			const FVertexID VertexID1 = GetEdgeVertex(EdgeID, 1);
+			VertexArray[VertexID0].ConnectedEdgeIDs.Add(EdgeID);
+			VertexArray[VertexID1].ConnectedEdgeIDs.Add(EdgeID);
+		}
+
+		if (Ar.CustomVer(FEditorObjectVersion::GUID) >= FEditorObjectVersion::MeshDescriptionTriangles)
+		{
+			// Make reverse connection from polygons to triangles
+			for (const FTriangleID TriangleID : TriangleArray.GetElementIDs())
+			{
+				const FPolygonID PolygonID = TriangleArray[TriangleID].PolygonID;
+				PolygonArray[PolygonID].TriangleIDs.Add(TriangleID);
+			}
 		}
 
 		// Populate polygon IDs for vertex instances, edges and polygon groups
-		for( const FPolygonID PolygonID : MeshDescription.PolygonArray.GetElementIDs() )
+		for (const FPolygonID PolygonID : PolygonArray.GetElementIDs())
 		{
-			auto PopulatePolygonIDs = [ &MeshDescription, PolygonID ]( const TArray<FVertexInstanceID>& VertexInstanceIDs )
+			if (Ar.CustomVer(FEditorObjectVersion::GUID) >= FEditorObjectVersion::MeshDescriptionTriangles)
 			{
-				const int32 NumVertexInstances = VertexInstanceIDs.Num();
-				for( int32 Index = 0; Index < NumVertexInstances; ++Index )
+				// If the polygon has no contour serialized, copy it over from the triangle
+				if (PolygonArray[PolygonID].PerimeterContour.VertexInstanceIDs.Num() == 0)
 				{
-					const FVertexInstanceID VertexInstanceID0 = VertexInstanceIDs[ Index ];
-					const FVertexInstanceID VertexInstanceID1 = VertexInstanceIDs[ ( Index + 1 ) % NumVertexInstances ];
-					const FVertexID VertexID0 = MeshDescription.GetVertexInstanceVertex( VertexInstanceID0 );
-					const FVertexID VertexID1 = MeshDescription.GetVertexInstanceVertex( VertexInstanceID1 );
-					const FEdgeID EdgeID = MeshDescription.GetVertexPairEdge( VertexID0, VertexID1 );
-					MeshDescription.VertexInstanceArray[ VertexInstanceID0 ].ConnectedPolygons.Add( PolygonID );
-					MeshDescription.EdgeArray[ EdgeID ].ConnectedPolygons.Add( PolygonID );
+					check(PolygonArray[PolygonID].TriangleIDs.Num() == 1);
+					const FTriangleID TriangleID = PolygonArray[PolygonID].TriangleIDs[0];
+					for (int32 Index = 0; Index < 3; Index++)
+					{
+						PolygonArray[PolygonID].PerimeterContour.VertexInstanceIDs.Add(TriangleArray[TriangleID].GetVertexInstanceID(Index));
+					}
 				}
-			};
+			}
 
-			PopulatePolygonIDs( MeshDescription.GetPolygonPerimeterVertexInstances( PolygonID ) );
+			// For backward compatibility only; in 4.24, this will be removed.
+			// Note: internal edges will have the same connected polygon twice.
+			const TArray<FVertexInstanceID>& VertexInstanceIDs = GetPolygonVertexInstances(PolygonID);
+			const int32 NumVertexInstances = VertexInstanceIDs.Num();
+			for (int32 Index = 0; Index < NumVertexInstances; ++Index)
+			{
+				const FVertexInstanceID VertexInstanceID0 = VertexInstanceIDs[Index];
+				const FVertexInstanceID VertexInstanceID1 = VertexInstanceIDs[(Index + 1) % NumVertexInstances];
+				const FVertexID VertexID0 = GetVertexInstanceVertex(VertexInstanceID0);
+				const FVertexID VertexID1 = GetVertexInstanceVertex(VertexInstanceID1);
+				const FEdgeID EdgeID = GetVertexPairEdge(VertexID0, VertexID1);
+				VertexInstanceArray[VertexInstanceID0].ConnectedPolygons.Add(PolygonID);
+				EdgeArray[EdgeID].ConnectedPolygons.Add(PolygonID);
+			}
+			// End backward compatibility
+
+			const FPolygonGroupID PolygonGroupID = PolygonArray[PolygonID].PolygonGroupID;
+			PolygonGroupArray[PolygonGroupID].Polygons.Add(PolygonID);
+		}
+	}
 
-			const FPolygonGroupID PolygonGroupID = MeshDescription.PolygonArray[ PolygonID ].PolygonGroupID;
-			MeshDescription.PolygonGroupArray[ PolygonGroupID ].Polygons.Add( PolygonID );
+	if (Ar.IsLoading())
+	{
+		if (Ar.CustomVer(FEditorObjectVersion::GUID) < FEditorObjectVersion::MeshDescriptionTriangles)
+		{
+			TriangleArray.Reset();
 
-			// We don't serialize triangles; instead the polygon gets retriangulated on load
-			MeshDescription.ComputePolygonTriangulation( PolygonID, MeshDescription.PolygonArray[ PolygonID ].Triangles );
+			// If we didn't serialize triangles, generate them from the polygon contour
+			for (const FPolygonID PolygonID : PolygonArray.GetElementIDs())
+			{
+				check(PolygonArray[PolygonID].TriangleIDs.Num() == 0);
+				ComputePolygonTriangulation(PolygonID);
+			}
+		}
+		else
+		{
+			// Otherwise connect existing triangles to vertex instances and edges
+			for (const FTriangleID TriangleID : TriangleArray.GetElementIDs())
+			{
+				for (int32 Index = 0; Index < 3; ++Index)
+				{
+					const FVertexInstanceID VertexInstanceID = GetTriangleVertexInstance(TriangleID, Index);
+					const FVertexInstanceID NextVertexInstanceID = GetTriangleVertexInstance(TriangleID, (Index + 1 == 3) ? 0 : Index + 1);
+
+					const FVertexID VertexID0 = GetVertexInstanceVertex(VertexInstanceID);
+					const FVertexID VertexID1 = GetVertexInstanceVertex(NextVertexInstanceID);
+
+					FEdgeID EdgeID = GetVertexPairEdge(VertexID0, VertexID1);
+					check(EdgeID != FEdgeID::Invalid);
+
+					VertexInstanceArray[VertexInstanceID].ConnectedTriangles.Add(TriangleID);
+					EdgeArray[EdgeID].ConnectedTriangles.Add(TriangleID);
+				}
+			}
 		}
 	}
 
-	return Ar;
+	if (Ar.IsLoading())
+	{
+		bool bAllTriangles = true;
+		int32 TriCount = 0;
+		for (const FPolygonID PolygonID : PolygonArray.GetElementIDs())
+		{
+			int32 NumTriangles = GetPolygonTriangleIDs(PolygonID).Num();
+			TriCount += NumTriangles;
+			if (NumTriangles > 3)
+			{
+				bAllTriangles = false;
+				UE_LOG(LogTemp, Log, TEXT("Polygon %d is a quad (or more)"), PolygonID.GetValue());
+			}
+		}
+
+		check(!bAllTriangles || TriCount == TriangleArray.Num());
+	}
 }
 
 
@@ -97,6 +179,7 @@ void FMeshDescription::Empty()
 	VertexArray.Reset();
 	VertexInstanceArray.Reset();
 	EdgeArray.Reset();
+	TriangleArray.Reset();
 	PolygonArray.Reset();
 	PolygonGroupArray.Reset();
 
@@ -104,6 +187,7 @@ void FMeshDescription::Empty()
 	VertexAttributesSet.Initialize( 0 );
 	VertexInstanceAttributesSet.Initialize( 0 );
 	EdgeAttributesSet.Initialize( 0 );
+	TriangleAttributesSet.Initialize( 0 );
 	PolygonAttributesSet.Initialize( 0 );
 	PolygonGroupAttributesSet.Initialize( 0 );
 }
@@ -114,6 +198,7 @@ bool FMeshDescription::IsEmpty() const
 	return VertexArray.GetArraySize() == 0 &&
 		   VertexInstanceArray.GetArraySize() == 0 &&
 		   EdgeArray.GetArraySize() == 0 &&
+		   TriangleArray.GetArraySize() == 0 &&
 		   PolygonArray.GetArraySize() == 0 &&
 		   PolygonGroupArray.GetArraySize() == 0;
 }
@@ -124,6 +209,7 @@ void FMeshDescription::Compact( FElementIDRemappings& OutRemappings )
 	VertexArray.Compact( OutRemappings.NewVertexIndexLookup );
 	VertexInstanceArray.Compact( OutRemappings.NewVertexInstanceIndexLookup );
 	EdgeArray.Compact( OutRemappings.NewEdgeIndexLookup );
+	TriangleArray.Compact( OutRemappings.NewTriangleIndexLookup );
 	PolygonArray.Compact( OutRemappings.NewPolygonIndexLookup );
 	PolygonGroupArray.Compact( OutRemappings.NewPolygonGroupIndexLookup );
 
@@ -137,6 +223,7 @@ void FMeshDescription::Remap( const FElementIDRemappings& Remappings )
 	VertexArray.Remap( Remappings.NewVertexIndexLookup );
 	VertexInstanceArray.Remap( Remappings.NewVertexInstanceIndexLookup );
 	EdgeArray.Remap( Remappings.NewEdgeIndexLookup );
+	TriangleArray.Remap( Remappings.NewTriangleIndexLookup );
 	PolygonArray.Remap( Remappings.NewPolygonIndexLookup );
 	PolygonGroupArray.Remap( Remappings.NewPolygonGroupIndexLookup );
 
@@ -145,6 +232,17 @@ void FMeshDescription::Remap( const FElementIDRemappings& Remappings )
 }
 
 
+void FMeshDescription::RemapAttributes( const FElementIDRemappings& Remappings )
+{
+	VertexAttributesSet.Remap( Remappings.NewVertexIndexLookup );
+	VertexInstanceAttributesSet.Remap( Remappings.NewVertexInstanceIndexLookup );
+	EdgeAttributesSet.Remap( Remappings.NewEdgeIndexLookup );
+	TriangleAttributesSet.Remap( Remappings.NewTriangleIndexLookup );
+	PolygonAttributesSet.Remap( Remappings.NewPolygonIndexLookup );
+	PolygonGroupAttributesSet.Remap( Remappings.NewPolygonGroupIndexLookup );
+}
+
+
 void FMeshDescription::FixUpElementIDs( const FElementIDRemappings& Remappings )
 {
 	for( const FVertexID VertexID : VertexArray.GetElementIDs() )
@@ -175,6 +273,11 @@ void FMeshDescription::FixUpElementIDs( const FElementIDRemappings& Remappings )
 		{
 			PolygonID = Remappings.GetRemappedPolygonID( PolygonID );
 		}
+
+		for( FTriangleID& TriangleID : VertexInstance.ConnectedTriangles )
+		{
+			TriangleID = Remappings.GetRemappedTriangleID( TriangleID );
+		}
 	}
 
 	for( const FEdgeID EdgeID : EdgeArray.GetElementIDs() )
@@ -187,11 +290,28 @@ void FMeshDescription::FixUpElementIDs( const FElementIDRemappings& Remappings )
 			Edge.VertexIDs[ Index ] = Remappings.GetRemappedVertexID( Edge.VertexIDs[ Index ] );
 		}
 
-		// Fix up references to section indices
 		for( FPolygonID& ConnectedPolygon : Edge.ConnectedPolygons )
 		{
 			ConnectedPolygon = Remappings.GetRemappedPolygonID( ConnectedPolygon );
 		}
+
+		for( FTriangleID& TriangleID : Edge.ConnectedTriangles )
+		{
+			TriangleID = Remappings.GetRemappedTriangleID( TriangleID );
+		}
+	}
+
+	for( const FTriangleID TriangleID : TriangleArray.GetElementIDs() )
+	{
+		FMeshTriangle& Triangle = TriangleArray[ TriangleID ];
+		
+		// Fix up vertex instance references in Triangle
+		for( int32 Index = 0; Index < 3; ++Index )
+		{
+			Triangle.SetVertexInstanceID( Index, Remappings.GetRemappedVertexInstanceID( Triangle.GetVertexInstanceID( Index ) ) );
+		}
+
+		Triangle.PolygonID = Remappings.GetRemappedPolygonID( Triangle.PolygonID );
 	}
 
 	for( const FPolygonID PolygonID : PolygonArray.GetElementIDs() )
@@ -204,14 +324,9 @@ void FMeshDescription::FixUpElementIDs( const FElementIDRemappings& Remappings )
 			VertexInstanceID = Remappings.GetRemappedVertexInstanceID( VertexInstanceID );
 		}
 
-		for( FMeshTriangle& Triangle : Polygon.Triangles )
+		for( FTriangleID& TriangleID : Polygon.TriangleIDs )
 		{
-			for( int32 TriangleVertexNumber = 0; TriangleVertexNumber < 3; ++TriangleVertexNumber )
-			{
-				const FVertexInstanceID OriginalVertexInstanceID = Triangle.GetVertexInstanceID( TriangleVertexNumber );
-				const FVertexInstanceID NewVertexInstanceID = Remappings.GetRemappedVertexInstanceID( OriginalVertexInstanceID );
-				Triangle.SetVertexInstanceID( TriangleVertexNumber, NewVertexInstanceID );
-			}
+			TriangleID = Remappings.GetRemappedTriangleID( TriangleID );
 		}
 
 		Polygon.PolygonGroupID = Remappings.GetRemappedPolygonGroupID( Polygon.PolygonGroupID );
@@ -229,91 +344,320 @@ void FMeshDescription::FixUpElementIDs( const FElementIDRemappings& Remappings )
 }
 
 
-void FMeshDescription::CreatePolygon_Internal( const FPolygonID PolygonID, const FPolygonGroupID PolygonGroupID, const TArray<FContourPoint>& Perimeter )
+void FMeshDescription::CreateVertexInstance_Internal(const FVertexInstanceID VertexInstanceID, const FVertexID VertexID)
+{
+	VertexInstanceArray[VertexInstanceID].VertexID = VertexID;
+	check(!VertexArray[VertexID].VertexInstanceIDs.Contains(VertexInstanceID));
+	VertexArray[VertexID].VertexInstanceIDs.Add(VertexInstanceID);
+	VertexInstanceAttributesSet.Insert(VertexInstanceID);
+}
+
+
+void FMeshDescription::DeleteVertexInstance(const FVertexInstanceID VertexInstanceID, TArray<FVertexID>* InOutOrphanedVerticesPtr)
+{
+	check(VertexInstanceArray[VertexInstanceID].ConnectedPolygons.Num() == 0);
+	const FVertexID VertexID = VertexInstanceArray[VertexInstanceID].VertexID;
+	verify(VertexArray[VertexID].VertexInstanceIDs.RemoveSingle(VertexInstanceID) == 1);
+	if (InOutOrphanedVerticesPtr && VertexArray[VertexID].VertexInstanceIDs.Num() == 0 && VertexArray[VertexID].ConnectedEdgeIDs.Num() == 0)
+	{
+		InOutOrphanedVerticesPtr->AddUnique(VertexID);
+	}
+	VertexInstanceArray.Remove(VertexInstanceID);
+	VertexInstanceAttributesSet.Remove(VertexInstanceID);
+}
+
+
+void FMeshDescription::CreateEdge_Internal(const FEdgeID EdgeID, const FVertexID VertexID0, const FVertexID VertexID1)
 {
-	FMeshPolygon& Polygon = PolygonArray[ PolygonID ];
+	check(GetVertexPairEdge(VertexID0, VertexID1) == FEdgeID::Invalid);
+	FMeshEdge& Edge = EdgeArray[EdgeID];
+	Edge.VertexIDs[0] = VertexID0;
+	Edge.VertexIDs[1] = VertexID1;
+	VertexArray[VertexID0].ConnectedEdgeIDs.Add(EdgeID);
+	VertexArray[VertexID1].ConnectedEdgeIDs.Add(EdgeID);
+	EdgeAttributesSet.Insert(EdgeID);
+}
+
 
-	Polygon.PerimeterContour.VertexInstanceIDs.Reset( Perimeter.Num() );
-	for( const FContourPoint ContourPoint : Perimeter )
+void FMeshDescription::DeleteEdge(const FEdgeID EdgeID, TArray<FVertexID>* InOutOrphanedVerticesPtr)
+{
+	FMeshEdge& Edge = EdgeArray[EdgeID];
+	for (const FVertexID EdgeVertexID : Edge.VertexIDs)
 	{
-		const FVertexInstanceID VertexInstanceID = ContourPoint.VertexInstanceID;
-		const FEdgeID EdgeID = ContourPoint.EdgeID;
+		FMeshVertex& Vertex = VertexArray[EdgeVertexID];
+		verify(Vertex.ConnectedEdgeIDs.RemoveSingle(EdgeID) == 1);
+		if (InOutOrphanedVerticesPtr && Vertex.ConnectedEdgeIDs.Num() == 0)
+		{
+			check(Vertex.VertexInstanceIDs.Num() == 0);  // We must already have deleted any vertex instances
+			InOutOrphanedVerticesPtr->AddUnique(EdgeVertexID);
+		}
+	}
+	EdgeArray.Remove(EdgeID);
+	EdgeAttributesSet.Remove(EdgeID);
+}
 
-		Polygon.PerimeterContour.VertexInstanceIDs.Add( VertexInstanceID );
-		check( !VertexInstanceArray[ VertexInstanceID ].ConnectedPolygons.Contains( PolygonID ) );
-		VertexInstanceArray[ VertexInstanceID ].ConnectedPolygons.Add( PolygonID );
 
-		check( !EdgeArray[ EdgeID ].ConnectedPolygons.Contains( PolygonID ) );
-		EdgeArray[ EdgeID ].ConnectedPolygons.Add( PolygonID );
+void FMeshDescription::CreateTriangle_Internal(const FTriangleID TriangleID, const FPolygonGroupID PolygonGroupID, TArrayView<const FVertexInstanceID> VertexInstanceIDs, TArray<FEdgeID>* OutEdgeIDs)
+{
+	if (OutEdgeIDs)
+	{
+		OutEdgeIDs->Empty();
 	}
 
+	// Fill out triangle vertex instances
+	FMeshTriangle& Triangle = TriangleArray[TriangleID];
+	check(VertexInstanceIDs.Num() == 3);
+	Triangle.SetVertexInstanceID(0, VertexInstanceIDs[0]);
+	Triangle.SetVertexInstanceID(1, VertexInstanceIDs[1]);
+	Triangle.SetVertexInstanceID(2, VertexInstanceIDs[2]);
+
+	// Make a polygon which will contain this triangle
+	FPolygonID PolygonID = PolygonArray.Add();
+	FMeshPolygon& Polygon = PolygonArray[PolygonID];
+	PolygonAttributesSet.Insert(PolygonID);
+
+	Polygon.PerimeterContour.VertexInstanceIDs.Reserve(3);
+	Algo::Copy(VertexInstanceIDs, Polygon.PerimeterContour.VertexInstanceIDs);
 	Polygon.PolygonGroupID = PolygonGroupID;
-	PolygonGroupArray[ PolygonGroupID ].Polygons.Add( PolygonID );
+	PolygonGroupArray[PolygonGroupID].Polygons.Add(PolygonID);
+
+	Triangle.PolygonID = PolygonID;
+	check(!Polygon.TriangleIDs.Contains(TriangleID));
+	Polygon.TriangleIDs.Add(TriangleID);
+
+	TriangleAttributesSet.Insert(TriangleID);
+
+	for (int32 Index = 0; Index < 3; ++Index)
+	{
+		const FVertexInstanceID VertexInstanceID = Triangle.GetVertexInstanceID(Index);
+		const FVertexInstanceID NextVertexInstanceID = Triangle.GetVertexInstanceID((Index == 2) ? 0 : Index + 1);
+
+		const FVertexID ThisVertexID = GetVertexInstanceVertex(VertexInstanceID);
+		const FVertexID NextVertexID = GetVertexInstanceVertex(NextVertexInstanceID);
+
+		FEdgeID EdgeID = GetVertexPairEdge(ThisVertexID, NextVertexID);
+		if (EdgeID == FEdgeID::Invalid)
+		{
+			EdgeID = CreateEdge(ThisVertexID, NextVertexID);
+			if (OutEdgeIDs)
+			{
+				OutEdgeIDs->Add(EdgeID);
+			}
+		}
+
+		check(!VertexInstanceArray[VertexInstanceID].ConnectedTriangles.Contains(TriangleID));
+		VertexInstanceArray[VertexInstanceID].ConnectedTriangles.Add(TriangleID);
+
+		check(!EdgeArray[EdgeID].ConnectedTriangles.Contains(TriangleID));
+		EdgeArray[EdgeID].ConnectedTriangles.Add(TriangleID);
 
-	PolygonAttributesSet.Insert( PolygonID );
+		// For backward compatibility only; in 4.24, this will be removed.
+		// Note: internal edges will have the same connected polygon twice.
+		VertexInstanceArray[VertexInstanceID].ConnectedPolygons.Add(PolygonID);
+		EdgeArray[EdgeID].ConnectedPolygons.Add(PolygonID);
+		// End backward compatibility
+	}
 }
 
 
-void FMeshDescription::CreatePolygon_Internal( const FPolygonID PolygonID, const FPolygonGroupID PolygonGroupID, const TArray<FVertexInstanceID>& VertexInstanceIDs, TArray<FEdgeID>* OutEdgeIDs )
+void FMeshDescription::DeleteTriangle(const FTriangleID TriangleID, TArray<FEdgeID>* InOutOrphanedEdgesPtr, TArray<FVertexInstanceID>* InOutOrphanedVertexInstancesPtr, TArray<FPolygonGroupID>* InOutOrphanedPolygonGroupsPtr)
 {
-	if( OutEdgeIDs )
+	const FMeshTriangle& Triangle = TriangleArray[TriangleID];
+	const FPolygonID PolygonID = Triangle.PolygonID;
+	
+	// Delete this triangle from the polygon
+	verify(PolygonArray[PolygonID].TriangleIDs.RemoveSingle(TriangleID) == 1);
+
+	if (PolygonArray[PolygonID].TriangleIDs.Num() == 0)
+	{
+		// If it was the only triangle in the polygon, delete the polygon too
+		for (int32 Index = 0; Index < 3; ++Index)
+		{
+			const FVertexInstanceID VertexInstanceID = Triangle.GetVertexInstanceID(Index);
+			const FVertexInstanceID NextVertexInstanceID = Triangle.GetVertexInstanceID((Index == 2) ? 0 : Index + 1);
+
+			const FVertexID VertexID0 = GetVertexInstanceVertex(VertexInstanceID);
+			const FVertexID VertexID1 = GetVertexInstanceVertex(NextVertexInstanceID);
+
+			FEdgeID EdgeID = GetVertexPairEdge(VertexID0, VertexID1);
+			check(EdgeID != FEdgeID::Invalid);
+
+			verify(VertexInstanceArray[VertexInstanceID].ConnectedTriangles.RemoveSingle(TriangleID) == 1);
+			verify(EdgeArray[EdgeID].ConnectedTriangles.RemoveSingle(TriangleID) == 1);
+
+			// For backward compatibility only; in 4.24, this will be removed.
+			// Note: internal edges will have the same connected polygon twice.
+			verify(VertexInstanceArray[VertexInstanceID].ConnectedPolygons.RemoveSingle(PolygonID) == 1);
+			verify(EdgeArray[EdgeID].ConnectedPolygons.RemoveSingle(PolygonID) == 1);
+			// End backward compatibility
+
+			if (InOutOrphanedVertexInstancesPtr && VertexInstanceArray[VertexInstanceID].ConnectedTriangles.Num() == 0)
+			{
+				InOutOrphanedVertexInstancesPtr->AddUnique(VertexInstanceID);
+			}
+
+			if (InOutOrphanedEdgesPtr && EdgeArray[EdgeID].ConnectedTriangles.Num() == 0)
+			{
+				InOutOrphanedEdgesPtr->AddUnique(EdgeID);
+			}
+		}
+
+		// Remove the polygon
+		const FPolygonGroupID PolygonGroupID = PolygonArray[PolygonID].PolygonGroupID;
+		verify(PolygonGroupArray[PolygonGroupID].Polygons.RemoveSingle(PolygonID) == 1);
+
+		if (InOutOrphanedPolygonGroupsPtr && PolygonGroupArray[PolygonGroupID].Polygons.Num() == 0)
+		{
+			InOutOrphanedPolygonGroupsPtr->AddUnique(PolygonGroupID);
+		}
+
+		PolygonArray.Remove( PolygonID );
+		PolygonAttributesSet.Remove( PolygonID );
+	}
+	else
+	{
+		// @todo: Handle this properly when deleting a triangle which forms part of an n-gon
+		// Either it needs to shave off the triangle from the contour and update the contour vertex instances,
+		// or it should just refuse to delete the triangle.
+		check(false);
+	}
+
+	TriangleArray.Remove(TriangleID);
+	TriangleAttributesSet.Remove(TriangleID);
+}
+
+
+void FMeshDescription::CreatePolygon_Internal(const FPolygonID PolygonID, const FPolygonGroupID PolygonGroupID, TArrayView<const FVertexInstanceID> VertexInstanceIDs, TArray<FEdgeID>* OutEdgeIDs)
+{
+	if (OutEdgeIDs)
 	{
 		OutEdgeIDs->Empty();
 	}
 
-	FMeshPolygon& Polygon = PolygonArray[ PolygonID ];
+	FMeshPolygon& Polygon = PolygonArray[PolygonID];
 	const int32 NumVertices = VertexInstanceIDs.Num();
-	Polygon.PerimeterContour.VertexInstanceIDs.Reset( NumVertices );
+	Polygon.PerimeterContour.VertexInstanceIDs.SetNumUninitialized(NumVertices);
 
-	for( int32 Index = 0; Index < NumVertices; ++Index )
+	for (int32 Index = 0; Index < NumVertices; ++Index)
 	{
-		const FVertexInstanceID VertexInstanceID = VertexInstanceIDs[ Index ];
-		const FVertexInstanceID NextVertexInstanceID = VertexInstanceIDs[ ( Index + 1 == NumVertices ) ? 0 : Index + 1 ];
-
-		Polygon.PerimeterContour.VertexInstanceIDs.Add( VertexInstanceID );
-		check( !VertexInstanceArray[ VertexInstanceID ].ConnectedPolygons.Contains( PolygonID ) );
-		VertexInstanceArray[ VertexInstanceID ].ConnectedPolygons.Add( PolygonID );
+		const FVertexInstanceID ThisVertexInstanceID = VertexInstanceIDs[Index];
+		const FVertexInstanceID NextVertexInstanceID = VertexInstanceIDs[(Index + 1 == NumVertices) ? 0 : Index + 1];
+		const FVertexID ThisVertexID = GetVertexInstanceVertex(ThisVertexInstanceID);
+		const FVertexID NextVertexID = GetVertexInstanceVertex(NextVertexInstanceID);
 
-		const FVertexID VertexID0 = GetVertexInstanceVertex( VertexInstanceID );
-		const FVertexID VertexID1 = GetVertexInstanceVertex( NextVertexInstanceID );
+		Polygon.PerimeterContour.VertexInstanceIDs[Index] = ThisVertexInstanceID;
 
-		FEdgeID EdgeID = GetVertexPairEdge( VertexID0, VertexID1 );
-		if( EdgeID == FEdgeID::Invalid )
+		FEdgeID EdgeID = GetVertexPairEdge(ThisVertexID, NextVertexID);
+		if (EdgeID == FEdgeID::Invalid)
 		{
-			EdgeID = CreateEdge( VertexID0, VertexID1 );
-			if( OutEdgeIDs )
+			EdgeID = CreateEdge(ThisVertexID, NextVertexID);
+			if (OutEdgeIDs)
 			{
-				OutEdgeIDs->Add( EdgeID );
+				OutEdgeIDs->Add(EdgeID);
 			}
 		}
 
-		check( !EdgeArray[ EdgeID ].ConnectedPolygons.Contains( PolygonID ) );
-		EdgeArray[ EdgeID ].ConnectedPolygons.Add( PolygonID );
+		// For backward compatibility only; in 4.24, this will be removed.
+		check(!VertexInstanceArray[ThisVertexInstanceID].ConnectedPolygons.Contains(PolygonID));
+		VertexInstanceArray[ThisVertexInstanceID].ConnectedPolygons.Add(PolygonID);
+		check(!EdgeArray[EdgeID].ConnectedPolygons.Contains(PolygonID));
+		EdgeArray[EdgeID].ConnectedPolygons.Add(PolygonID);
+		// End backward compatibility
 	}
 
 	Polygon.PolygonGroupID = PolygonGroupID;
-	PolygonGroupArray[ PolygonGroupID ].Polygons.Add( PolygonID );
+	PolygonGroupArray[PolygonGroupID].Polygons.Add(PolygonID);
 
-	PolygonAttributesSet.Insert( PolygonID );
-}
+	check(Polygon.TriangleIDs.Num() == 0);
+	ComputePolygonTriangulation(PolygonID);
 
+	PolygonAttributesSet.Insert(PolygonID);
+}
 
 
-void FMeshDescription::RemapAttributes( const FElementIDRemappings& Remappings )
+void FMeshDescription::DeletePolygon(const FPolygonID PolygonID, TArray<FEdgeID>* InOutOrphanedEdgesPtr, TArray<FVertexInstanceID>* InOutOrphanedVertexInstancesPtr, TArray<FPolygonGroupID>* InOutOrphanedPolygonGroupsPtr)
 {
-	VertexAttributesSet.Remap( Remappings.NewVertexIndexLookup );
-	VertexInstanceAttributesSet.Remap( Remappings.NewVertexInstanceIndexLookup );
-	EdgeAttributesSet.Remap( Remappings.NewEdgeIndexLookup );
-	PolygonAttributesSet.Remap( Remappings.NewPolygonIndexLookup );
-	PolygonGroupAttributesSet.Remap( Remappings.NewPolygonGroupIndexLookup );
+	FMeshPolygon& Polygon = PolygonArray[PolygonID];
+
+	// Remove constituent triangles
+	for (const FTriangleID TriangleID : Polygon.TriangleIDs)
+	{
+		const FMeshTriangle& Triangle = TriangleArray[TriangleID];
+
+		for (int32 Index = 0; Index < 3; ++Index)
+		{
+			const FVertexInstanceID ThisVertexInstanceID = Triangle.GetVertexInstanceID(Index);
+			const FVertexInstanceID NextVertexInstanceID = Triangle.GetVertexInstanceID((Index == 2) ? 0 : Index + 1);
+			const FVertexID ThisVertexID = GetVertexInstanceVertex(ThisVertexInstanceID);
+			const FVertexID NextVertexID = GetVertexInstanceVertex(NextVertexInstanceID);
+			const FEdgeID EdgeID = GetVertexPairEdge(ThisVertexID, NextVertexID);
+			
+			// If a valid edge isn't found, we deem this to be because it's an internal edge which was already removed
+			// in a previous iteration through the triangle array.
+			if (EdgeID != FEdgeID::Invalid)
+			{
+				if (IsEdgeInternal(EdgeID))
+				{
+					// Remove internal edges
+					for (const FVertexID EdgeVertexID : EdgeArray[EdgeID].VertexIDs)
+					{
+						verify(VertexArray[EdgeVertexID].ConnectedEdgeIDs.RemoveSingle(EdgeID) == 1);
+					}
+					EdgeArray.Remove(EdgeID);
+					EdgeAttributesSet.Remove(EdgeID);
+				}
+				else
+				{
+					verify(EdgeArray[EdgeID].ConnectedTriangles.RemoveSingle(TriangleID) == 1);
+
+					// For backward compatibility only; in 4.24, this will be removed.
+					verify(EdgeArray[EdgeID].ConnectedPolygons.RemoveSingle(PolygonID) == 1);
+					// End backward compatibility
+
+					if (InOutOrphanedEdgesPtr && EdgeArray[EdgeID].ConnectedTriangles.Num() == 0)
+					{
+						InOutOrphanedEdgesPtr->AddUnique(EdgeID);
+					}
+				}
+			}
+
+			verify(VertexInstanceArray[ThisVertexInstanceID].ConnectedTriangles.RemoveSingle(TriangleID) == 1);
+
+			if (InOutOrphanedVertexInstancesPtr && VertexInstanceArray[ThisVertexInstanceID].ConnectedTriangles.Num() == 0)
+			{
+				InOutOrphanedVertexInstancesPtr->AddUnique(ThisVertexInstanceID);
+			}
+		}
+
+		TriangleArray.Remove(TriangleID);
+		TriangleAttributesSet.Remove(TriangleID);
+	}
+
+	// For backward compatibility only; in 4.24, this will be removed.
+	for (const FVertexInstanceID VertexInstanceID : Polygon.PerimeterContour.VertexInstanceIDs)
+	{
+		verify(VertexInstanceArray[VertexInstanceID].ConnectedPolygons.RemoveSingle(PolygonID) == 1);
+	}
+	// End backward compatibility
+
+	FMeshPolygonGroup& PolygonGroup = PolygonGroupArray[ Polygon.PolygonGroupID ];
+	verify( PolygonGroup.Polygons.RemoveSingle( PolygonID ) == 1 );
+
+	if( InOutOrphanedPolygonGroupsPtr && PolygonGroup.Polygons.Num() == 0 )
+	{
+		InOutOrphanedPolygonGroupsPtr->AddUnique( Polygon.PolygonGroupID );
+	}
+
+	PolygonArray.Remove( PolygonID );
+	PolygonAttributesSet.Remove( PolygonID );
 }
 
 
 bool FMeshDescription::IsVertexOrphaned( const FVertexID VertexID ) const
 {
-	for( const FVertexInstanceID VertexInstanceID : GetVertex( VertexID ).VertexInstanceIDs )
+	for( const FVertexInstanceID VertexInstanceID : VertexArray[ VertexID ].VertexInstanceIDs )
 	{
-		if( GetVertexInstance( VertexInstanceID ).ConnectedPolygons.Num() > 0 )
+		if( VertexInstanceArray[ VertexInstanceID ].ConnectedPolygons.Num() > 0 )
 		{
 			return false;
 		}
@@ -323,46 +667,81 @@ bool FMeshDescription::IsVertexOrphaned( const FVertexID VertexID ) const
 }
 
 
-void FMeshDescription::GetPolygonPerimeterVertices( const FPolygonID PolygonID, TArray<FVertexID>& OutPolygonPerimeterVertexIDs ) const
+FEdgeID FMeshDescription::GetVertexPairEdge(const FVertexID VertexID0, const FVertexID VertexID1) const
 {
-	const FMeshPolygon& Polygon = GetPolygon( PolygonID );
-
-	OutPolygonPerimeterVertexIDs.SetNumUninitialized( Polygon.PerimeterContour.VertexInstanceIDs.Num(), false );
-
-	int32 Index = 0;
-	for( const FVertexInstanceID VertexInstanceID : Polygon.PerimeterContour.VertexInstanceIDs )
+	for (const FEdgeID VertexConnectedEdgeID : VertexArray[VertexID0].ConnectedEdgeIDs)
 	{
-		const FMeshVertexInstance& VertexInstance = VertexInstanceArray[ VertexInstanceID ];
-		OutPolygonPerimeterVertexIDs[ Index ] = VertexInstance.VertexID;
-		Index++;
+		const FVertexID EdgeVertexID0 = EdgeArray[VertexConnectedEdgeID].VertexIDs[0];
+		const FVertexID EdgeVertexID1 = EdgeArray[VertexConnectedEdgeID].VertexIDs[1];
+		if ((EdgeVertexID0 == VertexID0 && EdgeVertexID1 == VertexID1) || (EdgeVertexID0 == VertexID1 && EdgeVertexID1 == VertexID0))
+		{
+			return VertexConnectedEdgeID;
+		}
 	}
+
+	return FEdgeID::Invalid;
 }
 
 
-/** Given three direction vectors, indicates if A and B are on the same 'side' of Vec. */
-bool FMeshDescription::VectorsOnSameSide( const FVector& Vec, const FVector& A, const FVector& B, const float SameSideDotProductEpsilon )
+FEdgeID FMeshDescription::GetVertexInstancePairEdge(const FVertexInstanceID VertexInstanceID0, const FVertexInstanceID VertexInstanceID1) const
 {
-	const FVector CrossA = FVector::CrossProduct( Vec, A );
-	const FVector CrossB = FVector::CrossProduct( Vec, B );
-	float DotWithEpsilon = SameSideDotProductEpsilon + FVector::DotProduct( CrossA, CrossB );
-	return !FMath::IsNegativeFloat( DotWithEpsilon );
+	const FVertexID VertexID0 = VertexInstanceArray[VertexInstanceID0].VertexID;
+	const FVertexID VertexID1 = VertexInstanceArray[VertexInstanceID1].VertexID;
+	for (const FEdgeID VertexConnectedEdgeID : VertexArray[VertexID0].ConnectedEdgeIDs)
+	{
+		const FVertexID EdgeVertexID0 = EdgeArray[VertexConnectedEdgeID].VertexIDs[0];
+		const FVertexID EdgeVertexID1 = EdgeArray[VertexConnectedEdgeID].VertexIDs[1];
+		if ((EdgeVertexID0 == VertexID0 && EdgeVertexID1 == VertexID1) || (EdgeVertexID0 == VertexID1 && EdgeVertexID1 == VertexID0))
+		{
+			return VertexConnectedEdgeID;
+		}
+	}
+
+	return FEdgeID::Invalid;
 }
 
 
-/** Util to see if P lies within triangle created by A, B and C. */
-bool FMeshDescription::PointInTriangle( const FVector& A, const FVector& B, const FVector& C, const FVector& P, const float InsideTriangleDotProductEpsilon )
+TArray<FMeshTriangle> FMeshDescription::GetPolygonTriangles(const FPolygonID PolygonID) const
 {
-	// Cross product indicates which 'side' of the vector the point is on
-	// If its on the same side as the remaining vert for all edges, then its inside.	
-	if ( VectorsOnSameSide( B - A, P - A, C - A, InsideTriangleDotProductEpsilon ) &&
-		 VectorsOnSameSide( C - B, P - B, A - B, InsideTriangleDotProductEpsilon ) &&
-		 VectorsOnSameSide( A - C, P - C, B - C, InsideTriangleDotProductEpsilon ) )
+	TArray<FMeshTriangle> Triangles;
+	for (const FTriangleID TriangleID : PolygonArray[PolygonID].TriangleIDs)
 	{
-		return true;
+		Triangles.Add(TriangleArray[TriangleID]);
 	}
-	else
+	return Triangles;
+}
+
+
+void FMeshDescription::SetPolygonVertexInstance(const FPolygonID PolygonID, const int32 PerimeterIndex, const FVertexInstanceID VertexInstanceID)
+{
+	FMeshPolygon& Polygon = PolygonArray[PolygonID];
+	check(PerimeterIndex >= 0 && PerimeterIndex < Polygon.PerimeterContour.VertexInstanceIDs.Num());
+
+	// Disconnect old vertex instance from polygon, and connect new one
+	const FVertexInstanceID OldVertexInstanceID = Polygon.PerimeterContour.VertexInstanceIDs[PerimeterIndex];
+
+	// For backward compatibility only; in 4.24, this will be removed.
+	verify(VertexInstanceArray[OldVertexInstanceID].ConnectedPolygons.RemoveSingle(PolygonID) == 1);
+	check(!VertexInstanceArray[VertexInstanceID].ConnectedPolygons.Contains(PolygonID));
+	VertexInstanceArray[VertexInstanceID].ConnectedPolygons.Add(PolygonID);
+	// End backward compatibility
+
+	Polygon.PerimeterContour.VertexInstanceIDs[PerimeterIndex] = VertexInstanceID;
+
+	// Fix up triangle list
+	for (const FTriangleID TriangleID : Polygon.TriangleIDs)
 	{
-		return false;
+		FMeshTriangle& Triangle = TriangleArray[TriangleID];
+		for (int32 VertexIndex = 0; VertexIndex < 3; ++VertexIndex)
+		{
+			if (Triangle.GetVertexInstanceID(VertexIndex) == OldVertexInstanceID)
+			{
+				verify(VertexInstanceArray[OldVertexInstanceID].ConnectedTriangles.RemoveSingle(TriangleID) == 1);
+				check(!VertexInstanceArray[VertexInstanceID].ConnectedTriangles.Contains(TriangleID));
+				VertexInstanceArray[VertexInstanceID].ConnectedTriangles.Add(TriangleID);
+				Triangle.SetVertexInstanceID(VertexIndex, VertexInstanceID);
+			}
+		}
 	}
 }
 
@@ -381,7 +760,7 @@ FPlane FMeshDescription::ComputePolygonPlane( const FPolygonID PolygonID ) const
 	FVector Normal = FVector::ZeroVector;
 
 	static TArray<FVertexID> PerimeterVertexIDs;
-	GetPolygonPerimeterVertices( PolygonID, /* Out */ PerimeterVertexIDs );
+	GetPolygonVertices( PolygonID, /* Out */ PerimeterVertexIDs );
 
 	// @todo Maybe this shouldn't be in FMeshDescription but in a utility class, as it references a specific attribute name
 	TVertexAttributesConstRef<FVector> VertexPositions = VertexAttributes().GetAttributesRef<FVector>( MeshAttribute::Vertex::Position );
@@ -420,6 +799,275 @@ FVector FMeshDescription::ComputePolygonNormal( const FPolygonID PolygonID ) con
 }
 
 
+/** Returns true if the triangle formed by the specified three positions has a normal that is facing the opposite direction of the reference normal */
+static bool IsTriangleFlipped(const FVector ReferenceNormal, const FVector VertexPositionA, const FVector VertexPositionB, const FVector VertexPositionC)
+{
+	const FVector TriangleNormal = FVector::CrossProduct(
+		VertexPositionC - VertexPositionA,
+		VertexPositionB - VertexPositionA).GetSafeNormal();
+	return (FVector::DotProduct(ReferenceNormal, TriangleNormal) <= 0.0f);
+}
+
+
+/** Given three direction vectors, indicates if A and B are on the same 'side' of Vec. */
+static bool VectorsOnSameSide(const FVector& Vec, const FVector& A, const FVector& B, const float SameSideDotProductEpsilon)
+{
+	const FVector CrossA = FVector::CrossProduct(Vec, A);
+	const FVector CrossB = FVector::CrossProduct(Vec, B);
+	float DotWithEpsilon = SameSideDotProductEpsilon + FVector::DotProduct(CrossA, CrossB);
+	return !FMath::IsNegativeFloat(DotWithEpsilon);
+}
+
+
+/** Util to see if P lies within triangle created by A, B and C. */
+static bool PointInTriangle(const FVector& A, const FVector& B, const FVector& C, const FVector& P, const float InsideTriangleDotProductEpsilon)
+{
+	// Cross product indicates which 'side' of the vector the point is on
+	// If its on the same side as the remaining vert for all edges, then its inside.	
+	return (VectorsOnSameSide(B - A, P - A, C - A, InsideTriangleDotProductEpsilon) &&
+			VectorsOnSameSide(C - B, P - B, A - B, InsideTriangleDotProductEpsilon) &&
+			VectorsOnSameSide(A - C, P - C, B - C, InsideTriangleDotProductEpsilon));
+}
+
+
+void FMeshDescription::ComputePolygonTriangulation(const FPolygonID PolygonID)
+{
+	// NOTE: This polygon triangulation code is partially based on the ear cutting algorithm described on
+	//       page 497 of the book "Real-time Collision Detection", published in 2005.
+
+	FMeshPolygon& Polygon = PolygonArray[PolygonID];
+	const TArray<FVertexInstanceID>& PolygonVertexInstanceIDs = Polygon.PerimeterContour.VertexInstanceIDs;
+
+	// Polygon must have at least three vertices/edges
+	const int32 PolygonVertexCount = PolygonVertexInstanceIDs.Num();
+	check(PolygonVertexCount >= 3);
+
+	// If polygon was already triangulated, and only has three vertices, no need to do anything here
+	if (Polygon.TriangleIDs.Num() == 1 && PolygonVertexCount == 3)
+	{
+		return;
+	}
+
+	// Remove currently configured triangles
+	for (const FTriangleID TriangleID : Polygon.TriangleIDs)
+	{
+		// Disconnect triangles from vertex instances
+		for (const FVertexInstanceID VertexInstanceID : GetTriangleVertexInstances(TriangleID))
+		{
+			verify(VertexInstanceArray[VertexInstanceID].ConnectedTriangles.RemoveSingle(TriangleID) == 1);
+		}
+
+		// Disconnect triangles from perimeter edges, and delete internal edges
+		for (const FEdgeID EdgeID : GetTriangleEdges(TriangleID))
+		{
+			if (EdgeID != FEdgeID::Invalid)
+			{
+				// The edge may be invalid if it was an internal edge which was deleted in a previous iteration through the triangles.
+				// So only do something with valid edges
+				if (IsEdgeInternal(EdgeID))
+				{
+					// Remove internal edges completely (the first time they are seen)
+					for (const FVertexID VertexID : GetEdgeVertices(EdgeID))
+					{
+						// Disconnect edge from vertices
+						verify(VertexArray[VertexID].ConnectedEdgeIDs.RemoveSingle(EdgeID) == 1);
+					}
+
+					EdgeArray.Remove(EdgeID);
+					EdgeAttributesSet.Remove(EdgeID);
+				}
+				else
+				{
+					// Don't remove perimeter edge, but disconnect this triangle from it
+					verify(EdgeArray[EdgeID].ConnectedTriangles.RemoveSingle(TriangleID) == 1);
+				}
+			}
+		}
+
+		TriangleArray.Remove(TriangleID);
+		TriangleAttributesSet.Remove(TriangleID);
+	}
+
+	Polygon.TriangleIDs.Reset();
+
+	// If perimeter only has 3 vertices, just add a single triangle and return
+	if (PolygonVertexCount == 3)
+	{
+		const FTriangleID TriangleID = TriangleArray.Add();
+		TriangleAttributesSet.Insert(TriangleID);
+		
+		FMeshTriangle& Triangle = TriangleArray[TriangleID];
+		Triangle.PolygonID = PolygonID;
+		Polygon.TriangleIDs.Add(TriangleID);
+
+		for (int32 Index = 0; Index < 3; ++Index)
+		{
+			const FVertexInstanceID ThisVertexInstanceID = PolygonVertexInstanceIDs[Index];
+			const FVertexInstanceID NextVertexInstanceID = PolygonVertexInstanceIDs[(Index == 2) ? 0 : Index + 1];
+			const FVertexID ThisVertexID = GetVertexInstanceVertex(ThisVertexInstanceID);
+			const FVertexID NextVertexID = GetVertexInstanceVertex(NextVertexInstanceID);
+			const FEdgeID EdgeID = GetVertexPairEdge(ThisVertexID, NextVertexID);
+			check(EdgeID != FEdgeID::Invalid);
+
+			Triangle.SetVertexInstanceID(Index, ThisVertexInstanceID);
+
+			check(!EdgeArray[EdgeID].ConnectedTriangles.Contains(TriangleID));
+			EdgeArray[EdgeID].ConnectedTriangles.Add(TriangleID);
+
+			check(!VertexInstanceArray[ThisVertexInstanceID].ConnectedTriangles.Contains(TriangleID));
+			VertexInstanceArray[ThisVertexInstanceID].ConnectedTriangles.Add(TriangleID);
+		}
+
+		return;
+	}
+
+	// @todo mesheditor: Perhaps should always attempt to triangulate by splitting polygons along the shortest edge, for better determinism.
+
+	// First figure out the polygon normal.  We need this to determine which triangles are convex, so that
+	// we can figure out which ears to clip
+	const FVector PolygonNormal = ComputePolygonNormal(PolygonID);
+
+	// Make a simple linked list array of the previous and next vertex numbers, for each vertex number
+	// in the polygon.  This will just save us having to iterate later on.
+	TArray<int32> PrevVertexNumbers;
+	TArray<int32> NextVertexNumbers;
+	TArray<FVector> VertexPositions;
+
+	{
+		const TVertexAttributesRef<FVector> MeshVertexPositions = VertexAttributes().GetAttributesRef<FVector>(MeshAttribute::Vertex::Position);
+		PrevVertexNumbers.SetNumUninitialized(PolygonVertexCount, false);
+		NextVertexNumbers.SetNumUninitialized(PolygonVertexCount, false);
+		VertexPositions.SetNumUninitialized(PolygonVertexCount, false);
+
+		for (int32 VertexNumber = 0; VertexNumber < PolygonVertexCount; ++VertexNumber)
+		{
+			PrevVertexNumbers[VertexNumber] = VertexNumber - 1;
+			NextVertexNumbers[VertexNumber] = VertexNumber + 1;
+
+			VertexPositions[VertexNumber] = MeshVertexPositions[GetVertexInstanceVertex(PolygonVertexInstanceIDs[VertexNumber])];
+		}
+		PrevVertexNumbers[0] = PolygonVertexCount - 1;
+		NextVertexNumbers[PolygonVertexCount - 1] = 0;
+	}
+
+	int32 EarVertexNumber = 0;
+	int32 EarTestCount = 0;
+	for (int32 RemainingVertexCount = PolygonVertexCount; RemainingVertexCount >= 3; )
+	{
+		bool bIsEar = true;
+
+		// If we're down to only a triangle, just treat it as an ear.  Also, if we've tried every possible candidate
+		// vertex looking for an ear, go ahead and just treat the current vertex as an ear.  This can happen when 
+		// vertices are colinear or other degenerate cases.
+		if (RemainingVertexCount > 3 && EarTestCount < RemainingVertexCount)
+		{
+			const FVector PrevVertexPosition = VertexPositions[PrevVertexNumbers[EarVertexNumber]];
+			const FVector EarVertexPosition = VertexPositions[EarVertexNumber];
+			const FVector NextVertexPosition = VertexPositions[NextVertexNumbers[EarVertexNumber]];
+
+			// Figure out whether the potential ear triangle is facing the same direction as the polygon
+			// itself.  If it's facing the opposite direction, then we're dealing with a concave triangle
+			// and we'll skip it for now.
+			if (!IsTriangleFlipped(PolygonNormal, PrevVertexPosition, EarVertexPosition, NextVertexPosition))
+			{
+				int32 TestVertexNumber = NextVertexNumbers[NextVertexNumbers[EarVertexNumber]];
+
+				do
+				{
+					// Test every other remaining vertex to make sure that it doesn't lie inside our potential ear
+					// triangle.  If we find a vertex that's inside the triangle, then it cannot actually be an ear.
+					const FVector TestVertexPosition = VertexPositions[TestVertexNumber];
+					if (PointInTriangle(PrevVertexPosition, EarVertexPosition, NextVertexPosition, TestVertexPosition, SMALL_NUMBER))
+					{
+						bIsEar = false;
+						break;
+					}
+
+					TestVertexNumber = NextVertexNumbers[TestVertexNumber];
+
+				} while (TestVertexNumber != PrevVertexNumbers[EarVertexNumber]);
+			}
+			else
+			{
+				bIsEar = false;
+			}
+		}
+
+		if (bIsEar)
+		{
+			// OK, we found an ear!  Let's save this triangle in our output buffer.
+			// This will also create any missing internal edges.
+			{
+				// Add a new triangle
+				const FTriangleID TriangleID = TriangleArray.Add();
+				TriangleAttributesSet.Insert(TriangleID);
+
+				// Set its vertex instances and connect it to its parent polygon
+				FMeshTriangle& Triangle = TriangleArray[TriangleID];
+				Triangle.SetVertexInstanceID(0, PolygonVertexInstanceIDs[PrevVertexNumbers[EarVertexNumber]]);
+				Triangle.SetVertexInstanceID(1, PolygonVertexInstanceIDs[EarVertexNumber]);
+				Triangle.SetVertexInstanceID(2, PolygonVertexInstanceIDs[NextVertexNumbers[EarVertexNumber]]);
+				Triangle.PolygonID = PolygonID;
+				check(!Polygon.TriangleIDs.Contains(TriangleID));
+				Polygon.TriangleIDs.Add(TriangleID);
+
+				// Now generate internal edges and connected vertex instances to the new triangle
+				for (int32 Index = 0; Index < 3; ++Index)
+				{
+					const FVertexInstanceID ThisVertexInstanceID = Triangle.GetVertexInstanceID(Index);
+					const FVertexInstanceID NextVertexInstanceID = Triangle.GetVertexInstanceID((Index == 2) ? 0 : Index + 1);
+					const FVertexID ThisVertexID = GetVertexInstanceVertex(ThisVertexInstanceID);
+					const FVertexID NextVertexID = GetVertexInstanceVertex(NextVertexInstanceID);
+					FEdgeID EdgeID = GetVertexPairEdge(ThisVertexID, NextVertexID);
+					if (EdgeID == FEdgeID::Invalid)
+					{
+						// This must be an internal edge (as perimeter edges will already be defined)
+						EdgeID = CreateEdge(ThisVertexID, NextVertexID);
+
+						// Internal edges have the same polygon connected either side.
+						// For backward compatibility only; in 4.24, this will be removed.
+						EdgeArray[EdgeID].ConnectedPolygons.Add(PolygonID);
+						EdgeArray[EdgeID].ConnectedPolygons.Add(PolygonID);
+						// End backward compatibility
+					}
+
+					check(!VertexInstanceArray[ThisVertexInstanceID].ConnectedTriangles.Contains(TriangleID));
+					VertexInstanceArray[ThisVertexInstanceID].ConnectedTriangles.Add(TriangleID);
+
+					check(!EdgeArray[EdgeID].ConnectedTriangles.Contains(TriangleID));
+					EdgeArray[EdgeID].ConnectedTriangles.Add(TriangleID);
+				}
+			}
+
+			// Update our linked list.  We're effectively cutting off the ear by pointing the ear vertex's neighbors to
+			// point at their next sequential neighbor, and reducing the remaining vertex count by one.
+			{
+				NextVertexNumbers[PrevVertexNumbers[EarVertexNumber]] = NextVertexNumbers[EarVertexNumber];
+				PrevVertexNumbers[NextVertexNumbers[EarVertexNumber]] = PrevVertexNumbers[EarVertexNumber];
+				--RemainingVertexCount;
+			}
+
+			// Move on to the previous vertex in the list, now that this vertex was cut
+			EarVertexNumber = PrevVertexNumbers[EarVertexNumber];
+
+			EarTestCount = 0;
+		}
+		else
+		{
+			// The vertex is not the ear vertex, because it formed a triangle that either had a normal which pointed in the opposite direction
+			// of the polygon, or at least one of the other polygon vertices was found to be inside the triangle.  Move on to the next vertex.
+			EarVertexNumber = NextVertexNumbers[EarVertexNumber];
+
+			// Keep track of how many ear vertices we've tested, so that if we exhaust all remaining vertices, we can
+			// fall back to clipping the triangle and adding it to our mesh anyway.  This is important for degenerate cases.
+			++EarTestCount;
+		}
+	}
+
+	check(Polygon.TriangleIDs.Num() > 0);
+}
+
+
 void FMeshDescription::ComputePolygonTriangulation(const FPolygonID PolygonID, TArray<FMeshTriangle>& OutTriangles)
 {
 	// NOTE: This polygon triangulation code is partially based on the ear cutting algorithm described on
@@ -444,7 +1092,7 @@ void FMeshDescription::ComputePolygonTriangulation(const FPolygonID PolygonID, T
 	// @todo mesheditor: Perhaps should always attempt to triangulate by splitting polygons along the shortest edge, for better determinism.
 
 	//	const FMeshPolygon& Polygon = GetPolygon( PolygonID );
-	const TArray<FVertexInstanceID>& PolygonVertexInstanceIDs = GetPolygonPerimeterVertexInstances(PolygonID);
+	const TArray<FVertexInstanceID>& PolygonVertexInstanceIDs = GetPolygonVertexInstances(PolygonID);
 
 	// Polygon must have at least three vertices/edges
 	const int32 PolygonVertexCount = PolygonVertexInstanceIDs.Num();
@@ -586,8 +1234,7 @@ void FMeshDescription::TriangulateMesh()
 	// Perform triangulation directly into mesh polygons
 	for( const FPolygonID PolygonID : Polygons().GetElementIDs() )
 	{
-		FMeshPolygon& Polygon = PolygonArray[ PolygonID ];
-		ComputePolygonTriangulation( PolygonID, Polygon.Triangles );
+		ComputePolygonTriangulation( PolygonID );
 	}
 }
 
@@ -604,7 +1251,7 @@ bool FMeshDescription::ComputePolygonTangentsAndNormals(const FPolygonID Polygon
 	bool bValidNTBs = true;
 	// Calculate the center of this polygon
 	FVector Center = FVector::ZeroVector;
-	const TArray<FVertexInstanceID>& VertexInstanceIDs = GetPolygonPerimeterVertexInstances(PolygonID);
+	const TArray<FVertexInstanceID>& VertexInstanceIDs = GetPolygonVertexInstances(PolygonID);
 	for (const FVertexInstanceID VertexInstanceID : VertexInstanceIDs)
 	{
 		Center += VertexPositions[GetVertexInstanceVertex(VertexInstanceID)];
@@ -616,17 +1263,18 @@ bool FMeshDescription::ComputePolygonTangentsAndNormals(const FPolygonID Polygon
 	FVector Tangent = FVector::ZeroVector;
 	FVector Binormal = FVector::ZeroVector;
 
-	for (const FMeshTriangle& Triangle : GetPolygonTriangles(PolygonID))
+	for (const FTriangleID TriangleID : GetPolygonTriangleIDs(PolygonID))
 	{
-		const FVertexID VertexID0 = GetVertexInstanceVertex(Triangle.VertexInstanceID0);
-		const FVertexID VertexID1 = GetVertexInstanceVertex(Triangle.VertexInstanceID1);
-		const FVertexID VertexID2 = GetVertexInstanceVertex(Triangle.VertexInstanceID2);
+		const FMeshTriangle& Triangle = TriangleArray[ TriangleID ];
+		const FVertexID VertexID0 = GetVertexInstanceVertex(Triangle.GetVertexInstanceID(0));
+		const FVertexID VertexID1 = GetVertexInstanceVertex(Triangle.GetVertexInstanceID(1));
+		const FVertexID VertexID2 = GetVertexInstanceVertex(Triangle.GetVertexInstanceID(2));
 
 		const FVector DPosition1 = VertexPositions[VertexID1] - VertexPositions[VertexID0];
 		const FVector DPosition2 = VertexPositions[VertexID2] - VertexPositions[VertexID0];
 
-		const FVector2D DUV1 = VertexUVs[Triangle.VertexInstanceID1] - VertexUVs[Triangle.VertexInstanceID0];
-		const FVector2D DUV2 = VertexUVs[Triangle.VertexInstanceID2] - VertexUVs[Triangle.VertexInstanceID0];
+		const FVector2D DUV1 = VertexUVs[Triangle.GetVertexInstanceID(1)] - VertexUVs[Triangle.GetVertexInstanceID(0)];
+		const FVector2D DUV2 = VertexUVs[Triangle.GetVertexInstanceID(2)] - VertexUVs[Triangle.GetVertexInstanceID(0)];
 
 		// We have a left-handed coordinate system, but a counter-clockwise winding order
 		// Hence normal calculation has to take the triangle vectors cross product in reverse.
@@ -727,7 +1375,7 @@ void FMeshDescription::GetConnectedSoftEdges(const FVertexID VertexID, TArray<FE
 	OutConnectedSoftEdges.Reset();
 
 	TEdgeAttributesConstRef<bool> EdgeHardnesses = EdgeAttributes().GetAttributesRef<bool>(MeshAttribute::Edge::IsHard);
-	for (const FEdgeID ConnectedEdgeID : GetVertex(VertexID).ConnectedEdgeIDs)
+	for (const FEdgeID ConnectedEdgeID : VertexArray[VertexID].ConnectedEdgeIDs)
 	{
 		if (!EdgeHardnesses[ConnectedEdgeID])
 		{
@@ -803,7 +1451,7 @@ void FMeshDescription::GetVertexConnectedPolygonsInSameSoftEdgedGroup(const FVer
 
 float FMeshDescription::GetPolygonCornerAngleForVertex(const FPolygonID PolygonID, const FVertexID VertexID) const
 {
-	const FMeshPolygon& Polygon = GetPolygon(PolygonID);
+	const FMeshPolygon& Polygon = PolygonArray[PolygonID];
 
 	// Lambda function which returns the inner angle at a given index on a polygon contour
 	auto GetContourAngle = [this](const FMeshPolygonContour& Contour, const int32 ContourIndex)
@@ -1144,7 +1792,7 @@ void FMeshDescription::GetPolygonsInSameChartAsPolygon( const FPolygonID Polygon
 		Index++;
 
 		// Iterate through edges of the polygon
-		const TArray<FVertexInstanceID>& VertexInstanceIDs = GetPolygonPerimeterVertexInstances( PolygonToCheck );
+		const TArray<FVertexInstanceID>& VertexInstanceIDs = GetPolygonVertexInstances( PolygonToCheck );
 		FVertexID LastVertexID = GetVertexInstanceVertex( VertexInstanceIDs.Last() );
 		for( const FVertexInstanceID VertexInstanceID : VertexInstanceIDs )
 		{
@@ -1200,15 +1848,15 @@ void FMeshDescription::GetAllCharts( TArray<TArray<FPolygonID>>& OutCharts )
 
 void FMeshDescription::ReversePolygonFacing(const FPolygonID PolygonID)
 {
-	//Build a reverse perimeter
-	FMeshPolygon& Polygon = GetPolygon(PolygonID);
+	// Build a reverse perimeter
+	FMeshPolygon& Polygon = PolygonArray[PolygonID];
 	for (int32 i = 0; i < Polygon.PerimeterContour.VertexInstanceIDs.Num() / 2; ++i)
 	{
 		Polygon.PerimeterContour.VertexInstanceIDs.Swap(i, Polygon.PerimeterContour.VertexInstanceIDs.Num() - i - 1);
 	}
 	
-	//Triangulate the polygon since we reverse the indices
-	ComputePolygonTriangulation(PolygonID, Polygon.Triangles);
+	// Triangulate the polygon since we reverse the indices
+	ComputePolygonTriangulation(PolygonID);
 }
 
 
@@ -1222,6 +1870,52 @@ void FMeshDescription::ReverseAllPolygonFacing()
 }
 
 
+void FMeshDescription::RemapPolygonGroups(const TMap<FPolygonGroupID, FPolygonGroupID>& Remap)
+{
+	TPolygonGroupAttributesRef<FName> PolygonGroupNames = PolygonGroupAttributes().GetAttributesRef<FName>(MeshAttribute::PolygonGroup::ImportedMaterialSlotName);
+
+	struct FOldPolygonGroupData
+	{
+		FName Name;
+		TArray<FPolygonID> Polygons;
+	};
+
+	TMap<FPolygonGroupID, FOldPolygonGroupData> OldData;
+	for (const FPolygonGroupID& PolygonGroupID : PolygonGroups().GetElementIDs())
+	{
+		if (!Remap.Contains(PolygonGroupID) || PolygonGroupID == Remap[PolygonGroupID])
+		{
+			//No need to change this one
+			continue;
+		}
+		FOldPolygonGroupData& PolygonGroupData = OldData.FindOrAdd(PolygonGroupID);
+		PolygonGroupData.Name = PolygonGroupNames[PolygonGroupID];
+		FMeshPolygonGroup& PolygonGroup = PolygonGroupArray[PolygonGroupID];
+		PolygonGroupData.Polygons = PolygonGroup.Polygons;
+		PolygonGroup.Polygons.Empty();
+		DeletePolygonGroup(PolygonGroupID);
+	}
+	for (auto Kvp : OldData)
+	{
+		FPolygonGroupID GroupID = Kvp.Key;
+		FPolygonGroupID ToGroupID = Remap[GroupID];
+		if (!PolygonGroups().IsValid(ToGroupID))
+		{
+			CreatePolygonGroupWithID(ToGroupID);
+		}
+		TArray<FPolygonID>& Polygons = PolygonGroupArray[ToGroupID].Polygons;
+		Polygons.Append(Kvp.Value.Polygons);
+		PolygonGroupNames[ToGroupID] = Kvp.Value.Name;
+		for (FPolygonID PolygonID : Polygons)
+		{
+			PolygonArray[PolygonID].PolygonGroupID = ToGroupID;
+		}
+	}
+}
+
+PRAGMA_ENABLE_DEPRECATION_WARNINGS
+
+
 #if WITH_EDITORONLY_DATA
 
 void FMeshDescriptionBulkData::Serialize( FArchive& Ar, UObject* Owner )
diff --git a/Engine/Source/Runtime/MeshDescription/Private/MeshTypes.cpp b/Engine/Source/Runtime/MeshDescription/Private/MeshTypes.cpp
--- a/Engine/Source/Runtime/MeshDescription/Private/MeshTypes.cpp
+++ b/Engine/Source/Runtime/MeshDescription/Private/MeshTypes.cpp
@@ -9,6 +9,7 @@ const FElementID FElementID::Invalid( TNumericLimits<uint32>::Max() );
 const FVertexID FVertexID::Invalid( TNumericLimits<uint32>::Max() );
 const FVertexInstanceID FVertexInstanceID::Invalid( TNumericLimits<uint32>::Max() );
 const FEdgeID FEdgeID::Invalid( TNumericLimits<uint32>::Max() );
+const FTriangleID FTriangleID::Invalid(TNumericLimits<uint32>::Max());
 const FPolygonGroupID FPolygonGroupID::Invalid( TNumericLimits<uint32>::Max() );
 const FPolygonID FPolygonID::Invalid( TNumericLimits<uint32>::Max() );
 
