MeshDescription now uses triangles as its primary face representation. They can and should be used in preference to polygons when a mesh is known to be entirely triangulated, as they are quicker, and 

diff --git a/Engine/Source/Runtime/MeshDescription/Public/MeshDescription.h b/Engine/Source/Runtime/MeshDescription/Public/MeshDescription.h
--- a/Engine/Source/Runtime/MeshDescription/Public/MeshDescription.h
+++ b/Engine/Source/Runtime/MeshDescription/Public/MeshDescription.h
@@ -9,12 +9,17 @@
 #include "MeshTypes.h"
 #include "MeshElementArray.h"
 #include "MeshAttributeArray.h"
+#include "Algo/Accumulate.h"
+#include "Algo/Find.h"
+#include "Containers/ArrayView.h"
 #include "UObject/EditorObjectVersion.h"
 #include "UObject/ReleaseObjectVersion.h"
 #include "Serialization/BulkData.h"
 #include "Serialization/CustomVersion.h"
 #include "MeshDescription.generated.h"
 
+PRAGMA_DISABLE_DEPRECATION_WARNINGS
+
 enum
 {
 	//Remove the _MD when FRawMesh will be remove
@@ -23,13 +28,17 @@ enum
 
 struct FMeshVertex
 {
+	friend struct FMeshDescription;
+
 	FMeshVertex()
 	{}
 
 	/** All of vertex instances which reference this vertex (for split vertex support) */
+	UE_DEPRECATED(4.24, "This member should not be accessed directly. Please use FMeshDescription::GetVertexVertexInstances instead.")
 	TArray<FVertexInstanceID> VertexInstanceIDs;
 
 	/** The edges connected to this vertex */
+	UE_DEPRECATED(4.24, "This member should not be accessed directly. Please use FMeshDescription::GetVertexConnectedEdges instead.")
 	TArray<FEdgeID> ConnectedEdgeIDs;
 
 	/** Serializer */
@@ -48,16 +57,26 @@ struct FMeshVertex
 
 struct FMeshVertexInstance
 {
+	friend struct FMeshDescription;
+
 	FMeshVertexInstance()
 		: VertexID( FVertexID::Invalid )
 	{}
 
 	/** The vertex this is instancing */
+	UE_DEPRECATED(4.24, "This member should not be accessed directly. Please use FMeshDescription::GetVertexInstanceVertex instead.")
 	FVertexID VertexID;
 
-	/** List of connected polygons */
+	/** List of connected polygons. This will soon be deprecated entirely, in favour of the ConnectedTriangles list. */
+	UE_DEPRECATED(4.24, "This member should not be accessed directly. Please use FMeshDescription::GetVertexInstanceConnectedPolygons instead.")
 	TArray<FPolygonID> ConnectedPolygons;
 
+private:
+	/** List of connected polygons */
+	TArray<FTriangleID> ConnectedTriangles;
+
+public:
+
 	/** Serializer */
 	friend FArchive& operator<<( FArchive& Ar, FMeshVertexInstance& VertexInstance )
 	{
@@ -74,6 +93,8 @@ struct FMeshVertexInstance
 
 struct FMeshEdge
 {
+	friend struct FMeshDescription;
+
 	FMeshEdge()
 	{
 		VertexIDs[ 0 ] = FVertexID::Invalid;
@@ -81,12 +102,20 @@ struct FMeshEdge
 	}
 
 	/** IDs of the two editable mesh vertices that make up this edge.  The winding direction is not defined. */
+	UE_DEPRECATED(4.24, "This member should not be accessed directly. Please use FMeshDescription::GetEdgeVertex instead.")
 	FVertexID VertexIDs[ 2 ];
 
 	/** The polygons that share this edge.  It's best if there are always only two polygons that share
-	    the edge, and those polygons are facing the same direction */
+	    the edge, and those polygons are facing the same direction. This will soon be deprecated entirely, in favour of the ConnectedTriangles list. */
+	UE_DEPRECATED(4.24, "This member should not be accessed directly. Please use FMeshDescription::GetEdgeConnectedPolygons instead.")
 	TArray<FPolygonID> ConnectedPolygons;
 
+private:
+	/** The triangles that share this edge */
+	TArray<FTriangleID> ConnectedTriangles;
+
+public:
+
 	/** Serializer */
 	friend FArchive& operator<<( FArchive& Ar, FMeshEdge& Edge )
 	{
@@ -102,9 +131,10 @@ struct FMeshEdge
 };
 
 
-struct FMeshPolygonContour
+struct UE_DEPRECATED(4.24, "This class should no longer be used.") FMeshPolygonContour
 {
 	/** The ordered list of vertex instances which make up the polygon contour. The winding direction is counter-clockwise. */
+	UE_DEPRECATED(4.24, "This member should not be accessed directly. Please use FMeshDescription::GetPolygonPerimeterVertexInstances instead.")
 	TArray<FVertexInstanceID> VertexInstanceIDs;
 
 	/** Serializer */
@@ -116,10 +146,9 @@ struct FMeshPolygonContour
 };
 
 
-USTRUCT( BlueprintType )
 struct FMeshTriangle
 {
-	GENERATED_BODY()
+	friend struct FMeshDescription;
 
 	FMeshTriangle()
 		: VertexInstanceID0( FVertexInstanceID::Invalid ),
@@ -127,30 +156,36 @@ struct FMeshTriangle
 		  VertexInstanceID2( FVertexInstanceID::Invalid )
 	{}
 
-	/** First vertex instance that makes up this triangle.  Indices must be ordered counter-clockwise. */
-	UPROPERTY( BlueprintReadWrite, Category="Editable Mesh" )
+	/** Vertex instance IDs that make up this triangle.  Indices must be ordered counter-clockwise. */
+	// @todo: when public member access is deprecated, we will replace these clumsy members with a simple array
+	UE_DEPRECATED(4.24, "Please access triangle vertex instances through Get/SetVertexInstanceID or FMeshDescription::GetTriangleVertexID.")
 	FVertexInstanceID VertexInstanceID0;
 
-	/** Second vertex instance that makes up this triangle.  Indices must be ordered counter-clockwise. */
-	UPROPERTY( BlueprintReadWrite, Category="Editable Mesh" )
+	UE_DEPRECATED(4.24, "Please access triangle vertex instances through Get/SetVertexInstanceID or FMeshDescription::GetTriangleVertexID.")
 	FVertexInstanceID VertexInstanceID1;
 
-	/** Third vertex instance that makes up this triangle.  Indices must be ordered counter-clockwise. */
-	UPROPERTY( BlueprintReadWrite, Category="Editable Mesh" )
+	UE_DEPRECATED(4.24, "Please access triangle vertex instances through Get/SetVertexInstanceID or FMeshDescription::GetTriangleVertexID.")
 	FVertexInstanceID VertexInstanceID2;
 
+private:
+	/** Polygon which contains this triangle */
+	FPolygonID PolygonID;
+
+public:
 	/** Gets the specified triangle vertex instance ID.  Pass an index between 0 and 2 inclusive. */
 	inline FVertexInstanceID GetVertexInstanceID( const int32 Index ) const
 	{
+		// When we deprecate direct member access, this will be a simple array lookup
 		checkSlow( Index >= 0 && Index <= 2 );
-		return reinterpret_cast<const FVertexInstanceID*>( this )[ Index ];
+		return ( &VertexInstanceID0 )[ Index ];
 	}
 
 	/** Sets the specified triangle vertex instance ID.  Pass an index between 0 and 2 inclusive, and the new vertex instance ID to store. */
 	inline void SetVertexInstanceID( const int32 Index, const FVertexInstanceID NewVertexInstanceID )
 	{
+		// When we deprecate direct member access, this will be a simple array lookup
 		checkSlow( Index >= 0 && Index <= 2 );
-		( reinterpret_cast<FVertexInstanceID*>( this )[ Index ] ) = NewVertexInstanceID;
+		( &VertexInstanceID0 )[ Index ] = NewVertexInstanceID;
 	}
 
 	/** Serializer */
@@ -160,6 +195,11 @@ struct FMeshTriangle
 		Ar << Triangle.VertexInstanceID1;
 		Ar << Triangle.VertexInstanceID2;
 
+		if( !Ar.IsLoading() || Ar.CustomVer( FEditorObjectVersion::GUID ) >= FEditorObjectVersion::MeshDescriptionTriangles )
+		{
+			Ar << Triangle.PolygonID;
+		}
+
 		return Ar;
 	}
 };
@@ -167,23 +207,45 @@ struct FMeshTriangle
 
 struct FMeshPolygon
 {
+	friend struct FMeshDescription;
+
 	FMeshPolygon()
 		: PolygonGroupID( FPolygonGroupID::Invalid )
 	{}
 
 	/** The outer boundary edges of this polygon */
+	UE_DEPRECATED(4.24, "This member should not be accessed directly. Please use FMeshDescription::GetPolygonPerimeterVertexInstances instead.")
 	FMeshPolygonContour PerimeterContour;
 
 	/** List of triangles which make up this polygon */
+	UE_DEPRECATED(4.24, "This member is no longer used. Please use GetPolygonTriangleIDs instead.")
 	TArray<FMeshTriangle> Triangles;
 
+private:
+	/** List of triangle IDs which make up this polygon */
+	TArray<FTriangleID> TriangleIDs;
+
+public:
 	/** The polygon group which contains this polygon */
+	UE_DEPRECATED(4.24, "This member should not be accessed directly. Please use FMeshDescription::GetPolygonPolygonGroup instead.")
 	FPolygonGroupID PolygonGroupID;
 
 	/** Serializer */
 	friend FArchive& operator<<( FArchive& Ar, FMeshPolygon& Polygon )
 	{
-		Ar << Polygon.PerimeterContour;
+		if (Ar.IsSaving() &&
+			Ar.CustomVer(FEditorObjectVersion::GUID) >= FEditorObjectVersion::MeshDescriptionTriangles &&
+			Polygon.PerimeterContour.VertexInstanceIDs.Num() == 3)
+		{
+			// Optimisation: if polygon is a triangle, don't serialize the vertices as they can be copied over from the associated triangle
+			FMeshPolygonContour Empty;
+			Ar << Empty;
+		}
+		else
+		{
+			Ar << Polygon.PerimeterContour;
+		}
+
 		if (Ar.IsLoading() && Ar.CustomVer( FEditorObjectVersion::GUID ) < FEditorObjectVersion::MeshDescriptionRemovedHoles)
 		{
 			TArray<FMeshPolygonContour> Empty;
@@ -191,7 +253,8 @@ struct FMeshPolygon
 		}
 		if( Ar.IsLoading() && Ar.CustomVer( FReleaseObjectVersion::GUID ) < FReleaseObjectVersion::MeshDescriptionNewSerialization )
 		{
-			Ar << Polygon.Triangles;
+			TArray<FMeshTriangle> Triangles_DISCARD;
+			Ar << Triangles_DISCARD;
 		}
 		Ar << Polygon.PolygonGroupID;
 
@@ -202,10 +265,13 @@ struct FMeshPolygon
 
 struct FMeshPolygonGroup
 {
+	friend struct FMeshDescription;
+
 	FMeshPolygonGroup()
 	{}
 
 	/** All polygons in this group */
+	UE_DEPRECATED(4.24, "This member should not be accessed directly. Please use FMeshDescription::GetPolygonGroupPolygons instead.")
 	TArray<FPolygonID> Polygons;
 
 	/** Serializer */
@@ -225,6 +291,7 @@ struct FMeshPolygonGroup
 using FVertexArray = TMeshElementArray<FMeshVertex, FVertexID>;
 using FVertexInstanceArray = TMeshElementArray<FMeshVertexInstance, FVertexInstanceID>;
 using FEdgeArray = TMeshElementArray<FMeshEdge, FEdgeID>;
+using FTriangleArray = TMeshElementArray<FMeshTriangle, FTriangleID>;
 using FPolygonArray = TMeshElementArray<FMeshPolygon, FPolygonID>;
 using FPolygonGroupArray = TMeshElementArray<FMeshPolygonGroup, FPolygonGroupID>;
 
@@ -232,36 +299,42 @@ using FPolygonGroupArray = TMeshElementArray<FMeshPolygonGroup, FPolygonGroupID>
 template <typename AttributeType> using TVertexAttributeIndicesArray = TAttributeIndicesArray<AttributeType, FVertexID>;
 template <typename AttributeType> using TVertexInstanceAttributeIndicesArray = TAttributeIndicesArray<AttributeType, FVertexInstanceID>;
 template <typename AttributeType> using TEdgeAttributeIndicesArray = TAttributeIndicesArray<AttributeType, FEdgeID>;
+template <typename AttributeType> using TTriangleAttributeIndicesArray = TAttributeIndicesArray<AttributeType, FTriangleID>;
 template <typename AttributeType> using TPolygonAttributeIndicesArray = TAttributeIndicesArray<AttributeType, FPolygonID>;
 template <typename AttributeType> using TPolygonGroupAttributeIndicesArray = TAttributeIndicesArray<AttributeType, FPolygonGroupID>;
 
 template <typename AttributeType> using TVertexAttributeArray = TMeshAttributeArray<AttributeType, FVertexID>;
 template <typename AttributeType> using TVertexInstanceAttributeArray = TMeshAttributeArray<AttributeType, FVertexInstanceID>;
 template <typename AttributeType> using TEdgeAttributeArray = TMeshAttributeArray<AttributeType, FEdgeID>;
+template <typename AttributeType> using TTriangleAttributeArray = TMeshAttributeArray<AttributeType, FTriangleID>;
 template <typename AttributeType> using TPolygonAttributeArray = TMeshAttributeArray<AttributeType, FPolygonID>;
 template <typename AttributeType> using TPolygonGroupAttributeArray = TMeshAttributeArray<AttributeType, FPolygonGroupID>;
 
 template <typename AttributeType> using TVertexAttributesRef = TMeshAttributesRef<FVertexID, AttributeType>;
 template <typename AttributeType> using TVertexInstanceAttributesRef = TMeshAttributesRef<FVertexInstanceID, AttributeType>;
 template <typename AttributeType> using TEdgeAttributesRef = TMeshAttributesRef<FEdgeID, AttributeType>;
+template <typename AttributeType> using TTriangleAttributesRef = TMeshAttributesRef<FTriangleID, AttributeType>;
 template <typename AttributeType> using TPolygonAttributesRef = TMeshAttributesRef<FPolygonID, AttributeType>;
 template <typename AttributeType> using TPolygonGroupAttributesRef = TMeshAttributesRef<FPolygonGroupID, AttributeType>;
 
 template <typename AttributeType> using TVertexAttributesConstRef = TMeshAttributesConstRef<FVertexID, AttributeType>;
 template <typename AttributeType> using TVertexInstanceAttributesConstRef = TMeshAttributesConstRef<FVertexInstanceID, AttributeType>;
 template <typename AttributeType> using TEdgeAttributesConstRef = TMeshAttributesConstRef<FEdgeID, AttributeType>;
+template <typename AttributeType> using TTriangleAttributesConstRef = TMeshAttributesConstRef<FTriangleID, AttributeType>;
 template <typename AttributeType> using TPolygonAttributesConstRef = TMeshAttributesConstRef<FPolygonID, AttributeType>;
 template <typename AttributeType> using TPolygonGroupAttributesConstRef = TMeshAttributesConstRef<FPolygonGroupID, AttributeType>;
 
 template <typename AttributeType> using TVertexAttributesView = TMeshAttributesView<FVertexID, AttributeType>;
 template <typename AttributeType> using TVertexInstanceAttributesView = TMeshAttributesView<FVertexInstanceID, AttributeType>;
 template <typename AttributeType> using TEdgeAttributesView = TMeshAttributesView<FEdgeID, AttributeType>;
+template <typename AttributeType> using TTriangleAttributesView = TMeshAttributesView<FTriangleID, AttributeType>;
 template <typename AttributeType> using TPolygonAttributesView = TMeshAttributesView<FPolygonID, AttributeType>;
 template <typename AttributeType> using TPolygonGroupAttributesView = TMeshAttributesView<FPolygonGroupID, AttributeType>;
 
 template <typename AttributeType> using TVertexAttributesConstView = TMeshAttributesConstView<FVertexID, AttributeType>;
 template <typename AttributeType> using TVertexInstanceAttributesConstView = TMeshAttributesConstView<FVertexInstanceID, AttributeType>;
 template <typename AttributeType> using TEdgeAttributesConstView = TMeshAttributesConstView<FEdgeID, AttributeType>;
+template <typename AttributeType> using TTriangleAttributesConstView = TMeshAttributesConstView<FTriangleID, AttributeType>;
 template <typename AttributeType> using TPolygonAttributesConstView = TMeshAttributesConstView<FPolygonID, AttributeType>;
 template <typename AttributeType> using TPolygonGroupAttributesConstView = TMeshAttributesConstView<FPolygonGroupID, AttributeType>;
 
@@ -285,12 +358,19 @@ public:
 	// Hence explicitly declare all the below as defaulted, to ensure they will be generated by the compiler.
 	FMeshDescription() = default;
 	~FMeshDescription() = default;
-	FMeshDescription( const FMeshDescription& ) = default;
-	FMeshDescription( FMeshDescription&& ) = default;
-	FMeshDescription& operator=( const FMeshDescription& ) = default;
-	FMeshDescription& operator=( FMeshDescription&& ) = default;
+	FMeshDescription(const FMeshDescription&) = default;
+	FMeshDescription(FMeshDescription&&) = default;
+	FMeshDescription& operator=(const FMeshDescription&) = default;
+	FMeshDescription& operator=(FMeshDescription&&) = default;
 
-	friend MESHDESCRIPTION_API FArchive& operator<<( FArchive& Ar, FMeshDescription& MeshDescription );
+	friend MESHDESCRIPTION_API FArchive& operator<<(FArchive& Ar, FMeshDescription& MeshDescription)
+	{
+		MeshDescription.Serialize(Ar);
+		return Ar;
+	}
+
+	// Serialize the mesh description
+	void Serialize(FArchive& Ar);
 
 	// Empty the meshdescription
 	void Empty();
@@ -301,32 +381,45 @@ public:
 	FVertexArray& Vertices() { return VertexArray; }
 	const FVertexArray& Vertices() const { return VertexArray; }
 
-	FMeshVertex& GetVertex( const FVertexID VertexID ) { return VertexArray[ VertexID ]; }
-	const FMeshVertex& GetVertex( const FVertexID VertexID ) const { return VertexArray[ VertexID ]; }
+	UE_DEPRECATED(4.24, "Please do not access the vertex directly - use FMeshDescription accessor methods.")
+	FMeshVertex& GetVertex(const FVertexID VertexID) { return VertexArray[VertexID]; }
+	UE_DEPRECATED(4.24, "Please do not access the vertex directly - use FMeshDescription accessor methods.")
+	const FMeshVertex& GetVertex(const FVertexID VertexID) const { return VertexArray[VertexID]; }
 
 	FVertexInstanceArray& VertexInstances() { return VertexInstanceArray; }
 	const FVertexInstanceArray& VertexInstances() const { return VertexInstanceArray; }
 
-	FMeshVertexInstance& GetVertexInstance( const FVertexInstanceID VertexInstanceID ) { return VertexInstanceArray[ VertexInstanceID ]; }
-	const FMeshVertexInstance& GetVertexInstance( const FVertexInstanceID VertexInstanceID ) const { return VertexInstanceArray[ VertexInstanceID ]; }
+	UE_DEPRECATED(4.24, "Please do not access the vertex instance directly - use FMeshDescription accessor methods.")
+	FMeshVertexInstance& GetVertexInstance(const FVertexInstanceID VertexInstanceID) { return VertexInstanceArray[VertexInstanceID]; }
+	UE_DEPRECATED(4.24, "Please do not access the vertex instance directly - use FMeshDescription accessor methods.")
+	const FMeshVertexInstance& GetVertexInstance(const FVertexInstanceID VertexInstanceID) const { return VertexInstanceArray[VertexInstanceID]; }
 
 	FEdgeArray& Edges() { return EdgeArray; }
 	const FEdgeArray& Edges() const { return EdgeArray; }
 
-	FMeshEdge& GetEdge( const FEdgeID EdgeID ) { return EdgeArray[ EdgeID ]; }
-	const FMeshEdge& GetEdge( const FEdgeID EdgeID ) const { return EdgeArray[ EdgeID ]; }
+	UE_DEPRECATED(4.24, "Please do not access the edge directly - use FMeshDescription accessor methods.")
+	FMeshEdge& GetEdge(const FEdgeID EdgeID) { return EdgeArray[EdgeID]; }
+	UE_DEPRECATED(4.24, "Please do not access the edge directly - use FMeshDescription accessor methods.")
+	const FMeshEdge& GetEdge(const FEdgeID EdgeID) const { return EdgeArray[EdgeID]; }
+
+	FTriangleArray& Triangles() { return TriangleArray; }
+	const FTriangleArray& Triangles() const { return TriangleArray; }
 
 	FPolygonArray& Polygons() { return PolygonArray; }
 	const FPolygonArray& Polygons() const { return PolygonArray; }
 
-	FMeshPolygon& GetPolygon( const FPolygonID PolygonID ) { return PolygonArray[ PolygonID ]; }
-	const FMeshPolygon& GetPolygon( const FPolygonID PolygonID ) const { return PolygonArray[ PolygonID ]; }
+	UE_DEPRECATED(4.24, "Please do not access the polygon directly - use FMeshDescription accessor methods.")
+	FMeshPolygon& GetPolygon(const FPolygonID PolygonID) { return PolygonArray[PolygonID]; }
+	UE_DEPRECATED(4.24, "Please do not access the polygon directly - use FMeshDescription accessor methods.")
+	const FMeshPolygon& GetPolygon(const FPolygonID PolygonID) const { return PolygonArray[PolygonID]; }
 
 	FPolygonGroupArray& PolygonGroups() { return PolygonGroupArray; }
 	const FPolygonGroupArray& PolygonGroups() const { return PolygonGroupArray; }
 
-	FMeshPolygonGroup& GetPolygonGroup( const FPolygonGroupID PolygonGroupID ) { return PolygonGroupArray[ PolygonGroupID ]; }
-	const FMeshPolygonGroup& GetPolygonGroup( const FPolygonGroupID PolygonGroupID ) const { return PolygonGroupArray[ PolygonGroupID ]; }
+	UE_DEPRECATED(4.24, "Please do not access the polygon group directly - use FMeshDescription accessor methods.")
+	FMeshPolygonGroup& GetPolygonGroup(const FPolygonGroupID PolygonGroupID) { return PolygonGroupArray[PolygonGroupID]; }
+	UE_DEPRECATED(4.24, "Please do not access the polygon group directly - use FMeshDescription accessor methods.")
+	const FMeshPolygonGroup& GetPolygonGroup(const FPolygonGroupID PolygonGroupID) const { return PolygonGroupArray[PolygonGroupID]; }
 
 	TAttributesSet<FVertexID>& VertexAttributes() { return VertexAttributesSet; }
 	const TAttributesSet<FVertexID>& VertexAttributes() const { return VertexAttributesSet; }
@@ -337,483 +430,873 @@ public:
 	TAttributesSet<FEdgeID>& EdgeAttributes() { return EdgeAttributesSet; }
 	const TAttributesSet<FEdgeID>& EdgeAttributes() const { return EdgeAttributesSet; }
 
+	TAttributesSet<FTriangleID>& TriangleAttributes() { return TriangleAttributesSet; }
+	const TAttributesSet<FTriangleID>& TriangleAttributes() const { return TriangleAttributesSet; }
+
 	TAttributesSet<FPolygonID>& PolygonAttributes() { return PolygonAttributesSet; }
 	const TAttributesSet<FPolygonID>& PolygonAttributes() const { return PolygonAttributesSet; }
 
 	TAttributesSet<FPolygonGroupID>& PolygonGroupAttributes() { return PolygonGroupAttributesSet; }
 	const TAttributesSet<FPolygonGroupID>& PolygonGroupAttributes() const { return PolygonGroupAttributesSet; }
 
-	/** Reserves space for this number of new vertices */
-	void ReserveNewVertices( const int32 NumVertices )
-	{
-		VertexArray.Reserve( VertexArray.Num() + NumVertices );
-	}
 
-private:
-	void CreateVertex_Internal( const FVertexID VertexID )
+//////////////////////////////////////////////////////////////////////////
+// Create / remove mesh elements
+
+	/** Reserves space for this number of new vertices */
+	void ReserveNewVertices(const int32 NumVertices)
 	{
-		VertexAttributesSet.Insert( VertexID );
+		VertexArray.Reserve(VertexArray.Num() + NumVertices);
 	}
 
-public:
 	/** Adds a new vertex to the mesh and returns its ID */
 	FVertexID CreateVertex()
 	{
 		const FVertexID VertexID = VertexArray.Add();
-		CreateVertex_Internal( VertexID );
+		CreateVertex_Internal(VertexID);
 		return VertexID;
 	}
 
 	/** Adds a new vertex to the mesh with the given ID */
-	void CreateVertexWithID( const FVertexID VertexID )
+	void CreateVertexWithID(const FVertexID VertexID)
 	{
-		VertexArray.Insert( VertexID );
-		CreateVertex_Internal( VertexID );
+		VertexArray.Insert(VertexID);
+		CreateVertex_Internal(VertexID);
 	}
 
 	/** Deletes a vertex from the mesh */
-	void DeleteVertex( const FVertexID VertexID )
+	void DeleteVertex(const FVertexID VertexID)
 	{
-		check( VertexArray[ VertexID ].ConnectedEdgeIDs.Num() == 0 );
-		check( VertexArray[ VertexID ].VertexInstanceIDs.Num() == 0 );
-		VertexArray.Remove( VertexID );
-		VertexAttributesSet.Remove( VertexID );
+		check(VertexArray[VertexID].ConnectedEdgeIDs.Num() == 0);
+		check(VertexArray[VertexID].VertexInstanceIDs.Num() == 0);
+		VertexArray.Remove(VertexID);
+		VertexAttributesSet.Remove(VertexID);
 	}
 
 	/** Returns whether the passed vertex ID is valid */
-	bool IsVertexValid( const FVertexID VertexID ) const
+	bool IsVertexValid(const FVertexID VertexID) const
 	{
-		return VertexArray.IsValid( VertexID );
+		return VertexArray.IsValid(VertexID);
 	}
 
 	/** Reserves space for this number of new vertex instances */
-	void ReserveNewVertexInstances( const int32 NumVertexInstances )
+	void ReserveNewVertexInstances(const int32 NumVertexInstances)
 	{
-		VertexInstanceArray.Reserve( VertexInstanceArray.Num() + NumVertexInstances );
+		VertexInstanceArray.Reserve(VertexInstanceArray.Num() + NumVertexInstances);
 	}
 
-private:
-	void CreateVertexInstance_Internal( const FVertexInstanceID VertexInstanceID, const FVertexID VertexID )
-	{
-		VertexInstanceArray[ VertexInstanceID ].VertexID = VertexID;
-		check( !VertexArray[ VertexID ].VertexInstanceIDs.Contains( VertexInstanceID ) );
-		VertexArray[ VertexID ].VertexInstanceIDs.Add( VertexInstanceID );
-		VertexInstanceAttributesSet.Insert( VertexInstanceID );
-	}
-
-public:
 	/** Adds a new vertex instance to the mesh and returns its ID */
-	FVertexInstanceID CreateVertexInstance( const FVertexID VertexID )
+	FVertexInstanceID CreateVertexInstance(const FVertexID VertexID)
 	{
 		const FVertexInstanceID VertexInstanceID = VertexInstanceArray.Add();
-		CreateVertexInstance_Internal( VertexInstanceID, VertexID );
+		CreateVertexInstance_Internal(VertexInstanceID, VertexID);
 		return VertexInstanceID;
 	}
 
 	/** Adds a new vertex instance to the mesh with the given ID */
-	void CreateVertexInstanceWithID( const FVertexInstanceID VertexInstanceID, const FVertexID VertexID )
+	void CreateVertexInstanceWithID(const FVertexInstanceID VertexInstanceID, const FVertexID VertexID)
 	{
-		VertexInstanceArray.Insert( VertexInstanceID );
-		CreateVertexInstance_Internal( VertexInstanceID, VertexID );
+		VertexInstanceArray.Insert(VertexInstanceID);
+		CreateVertexInstance_Internal(VertexInstanceID, VertexID);
 	}
 
 	/** Deletes a vertex instance from a mesh */
-	void DeleteVertexInstance( const FVertexInstanceID VertexInstanceID, TArray<FVertexID>* InOutOrphanedVerticesPtr = nullptr )
-	{
-		check( VertexInstanceArray[ VertexInstanceID ].ConnectedPolygons.Num() == 0 );
-		const FVertexID VertexID = VertexInstanceArray[ VertexInstanceID ].VertexID;
-		verify( VertexArray[ VertexID ].VertexInstanceIDs.Remove( VertexInstanceID ) == 1 );
-		if( InOutOrphanedVerticesPtr && VertexArray[ VertexID ].VertexInstanceIDs.Num() == 0 && VertexArray[ VertexID ].ConnectedEdgeIDs.Num() == 0 )
-		{
-			InOutOrphanedVerticesPtr->AddUnique( VertexID );
-		}
-		VertexInstanceArray.Remove( VertexInstanceID );
-		VertexInstanceAttributesSet.Remove( VertexInstanceID );
-	}
+	void DeleteVertexInstance(const FVertexInstanceID VertexInstanceID, TArray<FVertexID>* InOutOrphanedVerticesPtr = nullptr);
 
 	/** Returns whether the passed vertex instance ID is valid */
-	bool IsVertexInstanceValid( const FVertexInstanceID VertexInstanceID ) const
+	bool IsVertexInstanceValid(const FVertexInstanceID VertexInstanceID) const
 	{
-		return VertexInstanceArray.IsValid( VertexInstanceID );
+		return VertexInstanceArray.IsValid(VertexInstanceID);
 	}
 
 	/** Reserves space for this number of new edges */
-	void ReserveNewEdges( const int32 NumEdges )
-	{
-		EdgeArray.Reserve( EdgeArray.Num() + NumEdges );
-	}
-
-private:
-	void CreateEdge_Internal( const FEdgeID EdgeID, const FVertexID VertexID0, const FVertexID VertexID1, const TArray<FPolygonID>& ConnectedPolygons )
+	void ReserveNewEdges(const int32 NumEdges)
 	{
-		FMeshEdge& Edge = EdgeArray[ EdgeID ];
-		Edge.VertexIDs[ 0 ] = VertexID0;
-		Edge.VertexIDs[ 1 ] = VertexID1;
-		Edge.ConnectedPolygons = ConnectedPolygons;
-		VertexArray[ VertexID0 ].ConnectedEdgeIDs.AddUnique( EdgeID );
-		VertexArray[ VertexID1 ].ConnectedEdgeIDs.AddUnique( EdgeID );
-		EdgeAttributesSet.Insert( EdgeID );
+		EdgeArray.Reserve(EdgeArray.Num() + NumEdges);
 	}
 
-public:
 	/** Adds a new edge to the mesh and returns its ID */
-	FEdgeID CreateEdge( const FVertexID VertexID0, const FVertexID VertexID1, const TArray<FPolygonID>& ConnectedPolygons = TArray<FPolygonID>() )
+	FEdgeID CreateEdge(const FVertexID VertexID0, const FVertexID VertexID1)
 	{
 		const FEdgeID EdgeID = EdgeArray.Add();
-		CreateEdge_Internal( EdgeID, VertexID0, VertexID1, ConnectedPolygons );
+		CreateEdge_Internal(EdgeID, VertexID0, VertexID1);
 		return EdgeID;
 	}
 
 	/** Adds a new edge to the mesh with the given ID */
-	void CreateEdgeWithID( const FEdgeID EdgeID, const FVertexID VertexID0, const FVertexID VertexID1, const TArray<FPolygonID>& ConnectedPolygons = TArray<FPolygonID>() )
+	void CreateEdgeWithID(const FEdgeID EdgeID, const FVertexID VertexID0, const FVertexID VertexID1)
 	{
-		EdgeArray.Insert( EdgeID );
-		CreateEdge_Internal( EdgeID, VertexID0, VertexID1, ConnectedPolygons );
+		EdgeArray.Insert(EdgeID);
+		CreateEdge_Internal(EdgeID, VertexID0, VertexID1);
 	}
 
 	/** Deletes an edge from a mesh */
-	void DeleteEdge( const FEdgeID EdgeID, TArray<FVertexID>* InOutOrphanedVerticesPtr = nullptr )
-	{
-		FMeshEdge& Edge = EdgeArray[ EdgeID ];
-		for( const FVertexID EdgeVertexID : Edge.VertexIDs )
-		{
-			FMeshVertex& Vertex = VertexArray[ EdgeVertexID ];
-			verify( Vertex.ConnectedEdgeIDs.RemoveSingle( EdgeID ) == 1 );
-			if( InOutOrphanedVerticesPtr && Vertex.ConnectedEdgeIDs.Num() == 0 )
-			{
-				check( Vertex.VertexInstanceIDs.Num() == 0 );  // We must already have deleted any vertex instances
-				InOutOrphanedVerticesPtr->AddUnique( EdgeVertexID );
-			}
-		}
-		EdgeArray.Remove( EdgeID );
-		EdgeAttributesSet.Remove( EdgeID );
-	}
+	void DeleteEdge(const FEdgeID EdgeID, TArray<FVertexID>* InOutOrphanedVerticesPtr = nullptr);
 
 	/** Returns whether the passed edge ID is valid */
-	bool IsEdgeValid( const FEdgeID EdgeID ) const
+	bool IsEdgeValid(const FEdgeID EdgeID) const
 	{
-		return EdgeArray.IsValid( EdgeID );
+		return EdgeArray.IsValid(EdgeID);
 	}
 
-	/** Reserves space for this number of new polygons */
-	void ReserveNewPolygons( const int32 NumPolygons )
+	/** Reserves space for this number of new triangles */
+	void ReserveNewTriangles(const int32 NumTriangles)
 	{
-		PolygonArray.Reserve( PolygonArray.Num() + NumPolygons );
+		TriangleArray.Reserve(TriangleArray.Num() + NumTriangles);
 	}
 
-	/** Pair of IDs representing the vertex instance ID on a contour, and the edge ID which starts at that point, winding counter-clockwise */
-	struct FContourPoint
+	/** Adds a new triangle to the mesh and returns its ID. This will also make an encapsulating polygon, and any missing edges. */
+	FTriangleID CreateTriangle(const FPolygonGroupID PolygonGroupID, TArrayView<const FVertexInstanceID> VertexInstanceIDs, TArray<FEdgeID>* OutEdgeIDs = nullptr)
 	{
-		FVertexInstanceID VertexInstanceID;
-		FEdgeID EdgeID;
-	};
-
-private:
-	void CreatePolygon_Internal( const FPolygonID PolygonID, const FPolygonGroupID PolygonGroupID, const TArray<FContourPoint>& Perimeter );
-	void CreatePolygon_Internal( const FPolygonID PolygonID, const FPolygonGroupID PolygonGroupID, const TArray<FVertexInstanceID>& VertexInstanceIDs, TArray<FEdgeID>* OutEdgeIDs );
+		const FTriangleID TriangleID = TriangleArray.Add();
+		CreateTriangle_Internal(TriangleID, PolygonGroupID, VertexInstanceIDs, OutEdgeIDs);
+		return TriangleID;
+	}
 
-public:
-	/** Adds a new polygon to the mesh and returns its ID */
-	UE_DEPRECATED(4.22, "Please use a different CreatePolygon() overload.")
-	FPolygonID CreatePolygon( const FPolygonGroupID PolygonGroupID, const TArray<FContourPoint>& Perimeter )
+	/** Adds a new triangle to the mesh with the given ID. This will also make an encapsulating polygon, and any missing edges. */
+	void CreateTriangleWithID(const FTriangleID TriangleID, const FPolygonGroupID PolygonGroupID, TArrayView<const FVertexInstanceID> VertexInstanceIDs, TArray<FEdgeID>* OutEdgeIDs = nullptr)
 	{
-		const FPolygonID PolygonID = PolygonArray.Add();
-		CreatePolygon_Internal( PolygonID, PolygonGroupID, Perimeter );
-		return PolygonID;
+		TriangleArray.Insert(TriangleID);
+		CreateTriangle_Internal(TriangleID, PolygonGroupID, VertexInstanceIDs, OutEdgeIDs);
 	}
 
-	FPolygonID CreatePolygon( const FPolygonGroupID PolygonGroupID, const TArray<FVertexInstanceID>& VertexInstanceIDs, TArray<FEdgeID>* OutEdgeIDs = nullptr )
+	/** Deletes a triangle from the mesh */
+	void DeleteTriangle(const FTriangleID TriangleID, TArray<FEdgeID>* InOutOrphanedEdgesPtr = nullptr, TArray<FVertexInstanceID>* InOutOrphanedVertexInstancesPtr = nullptr, TArray<FPolygonGroupID>* InOutOrphanedPolygonGroupsPtr = nullptr);
+
+	/** Returns whether the passed triangle ID is valid */
+	bool IsTriangleValid(const FTriangleID TriangleID) const
 	{
-		const FPolygonID PolygonID = PolygonArray.Add();
-		CreatePolygon_Internal( PolygonID, PolygonGroupID, VertexInstanceIDs, OutEdgeIDs );
-		return PolygonID;
+		return TriangleArray.IsValid(TriangleID);
 	}
 
-	/** Adds a new polygon to the mesh with the given ID */
-	UE_DEPRECATED(4.22, "Please use a different CreatePolygonWithID() overload.")
-	void CreatePolygonWithID( const FPolygonID PolygonID, const FPolygonGroupID PolygonGroupID, const TArray<FContourPoint>& Perimeter )
+	/** Reserves space for this number of new polygons */
+	void ReserveNewPolygons(const int32 NumPolygons)
 	{
-		PolygonArray.Insert( PolygonID );
-		CreatePolygon_Internal( PolygonID, PolygonGroupID, Perimeter );
+		PolygonArray.Reserve(PolygonArray.Num() + NumPolygons);
 	}
 
-	void CreatePolygonWithID( const FPolygonID PolygonID, const FPolygonGroupID PolygonGroupID, const TArray<FVertexInstanceID>& VertexInstanceIDs, TArray<FEdgeID>* OutEdgeIDs = nullptr )
+	/** Adds a new polygon to the mesh and returns its ID. This will also make any missing edges, and all constituent triangles. */
+	FPolygonID CreatePolygon(const FPolygonGroupID PolygonGroupID, TArrayView<const FVertexInstanceID> VertexInstanceIDs, TArray<FEdgeID>* OutEdgeIDs = nullptr)
 	{
-		PolygonArray.Insert( PolygonID );
-		CreatePolygon_Internal( PolygonID, PolygonGroupID, VertexInstanceIDs, OutEdgeIDs );
+		const FPolygonID PolygonID = PolygonArray.Add();
+		CreatePolygon_Internal(PolygonID, PolygonGroupID, VertexInstanceIDs, OutEdgeIDs);
+		return PolygonID;
 	}
 
-private:
-	void DeletePolygonContour_Internal( const FPolygonID PolygonID, const TArray<FVertexInstanceID>& VertexInstanceIDs, TArray<FEdgeID>* InOutOrphanedEdgesPtr, TArray<FVertexInstanceID>* InOutOrphanedVertexInstancesPtr )
+	/** Adds a new polygon to the mesh with the given ID. This will also make any missing edges, and all constituent triangles. */
+	void CreatePolygonWithID(const FPolygonID PolygonID, const FPolygonGroupID PolygonGroupID, TArrayView<const FVertexInstanceID> VertexInstanceIDs, TArray<FEdgeID>* OutEdgeIDs = nullptr)
 	{
-		FVertexInstanceID LastVertexInstanceID = VertexInstanceIDs.Last();
-		for( const FVertexInstanceID VertexInstanceID : VertexInstanceIDs )
-		{
-			FMeshVertexInstance& VertexInstance = VertexInstanceArray[ VertexInstanceID ];
-			verify( VertexInstance.ConnectedPolygons.Remove( PolygonID ) == 1 );
-
-			if( InOutOrphanedVertexInstancesPtr && VertexInstance.ConnectedPolygons.Num() == 0 )
-			{
-				InOutOrphanedVertexInstancesPtr->AddUnique( VertexInstanceID );
-			}
-
-			const FEdgeID EdgeID = GetVertexPairEdge( VertexInstanceArray[ LastVertexInstanceID ].VertexID, VertexInstanceArray[ VertexInstanceID ].VertexID );
-			check( EdgeID != FEdgeID::Invalid );
-			
-			FMeshEdge& Edge = EdgeArray[ EdgeID ];
-			verify( Edge.ConnectedPolygons.Remove( PolygonID ) == 1 );
-
-			if( InOutOrphanedEdgesPtr && Edge.ConnectedPolygons.Num() == 0 )
-			{
-				InOutOrphanedEdgesPtr->AddUnique( EdgeID );
-			}
-
-			LastVertexInstanceID = VertexInstanceID;
-		}
+		PolygonArray.Insert(PolygonID);
+		CreatePolygon_Internal(PolygonID, PolygonGroupID, VertexInstanceIDs, OutEdgeIDs);
 	}
 
-public:
 	/** Deletes a polygon from the mesh */
-	void DeletePolygon( const FPolygonID PolygonID, TArray<FEdgeID>* InOutOrphanedEdgesPtr = nullptr, TArray<FVertexInstanceID>* InOutOrphanedVertexInstancesPtr = nullptr, TArray<FPolygonGroupID>* InOutOrphanedPolygonGroupsPtr = nullptr )
-	{
-		FMeshPolygon& Polygon = PolygonArray[ PolygonID ];
-		DeletePolygonContour_Internal( PolygonID, Polygon.PerimeterContour.VertexInstanceIDs, InOutOrphanedEdgesPtr, InOutOrphanedVertexInstancesPtr );
-
-		FMeshPolygonGroup& PolygonGroup = PolygonGroupArray[ Polygon.PolygonGroupID ];
-		verify( PolygonGroup.Polygons.Remove( PolygonID ) == 1 );
-
-		if( InOutOrphanedPolygonGroupsPtr && PolygonGroup.Polygons.Num() == 0 )
-		{
-			InOutOrphanedPolygonGroupsPtr->AddUnique( Polygon.PolygonGroupID );
-		}
-
-		PolygonArray.Remove( PolygonID );
-		PolygonAttributesSet.Remove( PolygonID );
-	}
+	void DeletePolygon(const FPolygonID PolygonID, TArray<FEdgeID>* InOutOrphanedEdgesPtr = nullptr, TArray<FVertexInstanceID>* InOutOrphanedVertexInstancesPtr = nullptr, TArray<FPolygonGroupID>* InOutOrphanedPolygonGroupsPtr = nullptr);
 
 	/** Returns whether the passed polygon ID is valid */
-	bool IsPolygonValid( const FPolygonID PolygonID ) const
+	bool IsPolygonValid(const FPolygonID PolygonID) const
 	{
-		return PolygonArray.IsValid( PolygonID );
+		return PolygonArray.IsValid(PolygonID);
 	}
 
 	/** Reserves space for this number of new polygon groups */
-	void ReserveNewPolygonGroups( const int32 NumPolygonGroups )
-	{
-		PolygonGroupArray.Reserve( PolygonGroupArray.Num() + NumPolygonGroups );
-	}
-
-private:
-	void CreatePolygonGroup_Internal( const FPolygonGroupID PolygonGroupID )
+	void ReserveNewPolygonGroups(const int32 NumPolygonGroups)
 	{
-		PolygonGroupAttributesSet.Insert( PolygonGroupID );
+		PolygonGroupArray.Reserve(PolygonGroupArray.Num() + NumPolygonGroups);
 	}
 
-public:
 	/** Adds a new polygon group to the mesh and returns its ID */
 	FPolygonGroupID CreatePolygonGroup()
 	{
 		const FPolygonGroupID PolygonGroupID = PolygonGroupArray.Add();
-		CreatePolygonGroup_Internal( PolygonGroupID );
+		CreatePolygonGroup_Internal(PolygonGroupID);
 		return PolygonGroupID;
 	}
 
 	/** Adds a new polygon group to the mesh with the given ID */
-	void CreatePolygonGroupWithID( const FPolygonGroupID PolygonGroupID )
+	void CreatePolygonGroupWithID(const FPolygonGroupID PolygonGroupID)
 	{
-		PolygonGroupArray.Insert( PolygonGroupID );
-		CreatePolygonGroup_Internal( PolygonGroupID );
+		PolygonGroupArray.Insert(PolygonGroupID);
+		CreatePolygonGroup_Internal(PolygonGroupID);
 	}
 
 	/** Deletes a polygon group from the mesh */
-	void DeletePolygonGroup( const FPolygonGroupID PolygonGroupID )
+	void DeletePolygonGroup(const FPolygonGroupID PolygonGroupID)
 	{
-		check( PolygonGroupArray[ PolygonGroupID ].Polygons.Num() == 0 );
-		PolygonGroupArray.Remove( PolygonGroupID );
-		PolygonGroupAttributesSet.Remove( PolygonGroupID );
+		check(PolygonGroupArray[PolygonGroupID].Polygons.Num() == 0);
+		PolygonGroupArray.Remove(PolygonGroupID);
+		PolygonGroupAttributesSet.Remove(PolygonGroupID);
 	}
 
 	/** Returns whether the passed polygon group ID is valid */
-	bool IsPolygonGroupValid( const FPolygonGroupID PolygonGroupID ) const
+	bool IsPolygonGroupValid(const FPolygonGroupID PolygonGroupID) const
 	{
-		return PolygonGroupArray.IsValid( PolygonGroupID );
+		return PolygonGroupArray.IsValid(PolygonGroupID);
 	}
 
+
 //////////////////////////////////////////////////////////////////////////
-// Meshdescription general functions
+// MeshDescription general functions
 
 public:
 
+	//////////////////////////////////////////////////////////////////////
+	// Vertex operations
+
 	/** Returns whether a given vertex is orphaned, i.e. it doesn't form part of any polygon */
-	bool IsVertexOrphaned( const FVertexID VertexID ) const;
+	bool IsVertexOrphaned(const FVertexID VertexID) const;
 
 	/** Returns the edge ID defined by the two given vertex IDs, if there is one; otherwise FEdgeID::Invalid */
-	FEdgeID GetVertexPairEdge( const FVertexID VertexID0, const FVertexID VertexID1 ) const
+	FEdgeID GetVertexPairEdge(const FVertexID VertexID0, const FVertexID VertexID1) const;
+
+	/** Returns reference to an array of Edge IDs connected to this vertex */
+	const TArray<FEdgeID>& GetVertexConnectedEdges(const FVertexID VertexID) const
+	{
+		return VertexArray[VertexID].ConnectedEdgeIDs;
+	}
+
+	/** Returns number of edges connected to this vertex */
+	int32 GetNumVertexConnectedEdges(const FVertexID VertexID) const
 	{
-		for( const FEdgeID VertexConnectedEdgeID : VertexArray[ VertexID0 ].ConnectedEdgeIDs )
+		return VertexArray[VertexID].ConnectedEdgeIDs.Num();
+	}
+
+	/** Returns reference to an array of VertexInstance IDs instanced from this vertex */
+	const TArray<FVertexInstanceID>& GetVertexVertexInstances(const FVertexID VertexID) const
+	{
+		return VertexArray[VertexID].VertexInstanceIDs;
+	}
+
+	/** Returns number of vertex instances created from this vertex */
+	int32 GetNumVertexVertexInstances(const FVertexID VertexID) const
+	{
+		return VertexArray[VertexID].VertexInstanceIDs.Num();
+	}
+
+	/** Populates the passed array of TriangleIDs with the triangles connected to this vertex */
+	template <typename Alloc>
+	void GetVertexConnectedTriangles(const FVertexID VertexID, TArray<FTriangleID, Alloc>& OutConnectedTriangleIDs) const
+	{
+		OutConnectedTriangleIDs.Reset(GetNumVertexConnectedTriangles(VertexID));
+		for (const FVertexInstanceID VertexInstanceID : VertexArray[VertexID].VertexInstanceIDs)
 		{
-			const FVertexID EdgeVertexID0 = EdgeArray[ VertexConnectedEdgeID ].VertexIDs[ 0 ];
-			const FVertexID EdgeVertexID1 = EdgeArray[ VertexConnectedEdgeID ].VertexIDs[ 1 ];
-			if( ( EdgeVertexID0 == VertexID0 && EdgeVertexID1 == VertexID1 ) || ( EdgeVertexID0 == VertexID1 && EdgeVertexID1 == VertexID0 ) )
-			{
-				return VertexConnectedEdgeID;
-			}
+			OutConnectedTriangleIDs.Append(VertexInstanceArray[VertexInstanceID].ConnectedTriangles);
 		}
+	}
 
-		return FEdgeID::Invalid;
+	/** Returns the triangles connected to this vertex as an array with the specified allocator template type. */
+	template <typename Alloc>
+	TArray<FTriangleID, Alloc> GetVertexConnectedTriangles(const FVertexID VertexID) const
+	{
+		TArray<FTriangleID, Alloc> Result;
+		this->GetVertexConnectedTriangles(VertexID, Result);
+		return Result;
 	}
 
-	/** Returns reference to an array of Edge IDs connected to this vertex */
-	const TArray<FEdgeID>& GetVertexConnectedEdges( const FVertexID VertexID ) const
+	/** Returns the triangles connected to this vertex */
+	TArray<FTriangleID> GetVertexConnectedTriangles(const FVertexID VertexID) const
 	{
-		return VertexArray[ VertexID ].ConnectedEdgeIDs;
+		TArray<FTriangleID> Result;
+		this->GetVertexConnectedTriangles(VertexID, Result);
+		return Result;
 	}
 
-	/** Returns reference to an array of VertexInstance IDs instanced from this vertex */
-	const TArray<FVertexInstanceID>& GetVertexVertexInstances( const FVertexID VertexID ) const
+	/** Returns number of triangles connected to this vertex */
+	int32 GetNumVertexConnectedTriangles(const FVertexID VertexID) const
 	{
-		return VertexArray[ VertexID ].VertexInstanceIDs;
+		return Algo::TransformAccumulate(VertexArray[VertexID].VertexInstanceIDs,
+										 [this](const FVertexInstanceID ID) { return VertexInstanceArray[ID].ConnectedTriangles.Num(); },
+										 0);
 	}
 
 	/** Populates the passed array of PolygonIDs with the polygons connected to this vertex */
-	void GetVertexConnectedPolygons( const FVertexID VertexID, TArray<FPolygonID>& OutConnectedPolygonIDs ) const
+	template <typename Alloc>
+	void GetVertexConnectedPolygons(const FVertexID VertexID, TArray<FPolygonID, Alloc>& OutConnectedPolygonIDs) const
 	{
 		OutConnectedPolygonIDs.Reset();
-		for( const FVertexInstanceID VertexInstanceID : VertexArray[ VertexID ].VertexInstanceIDs )
+		for (const FVertexInstanceID VertexInstanceID : VertexArray[VertexID].VertexInstanceIDs)
 		{
-			OutConnectedPolygonIDs.Append( VertexInstanceArray[ VertexInstanceID ].ConnectedPolygons );
+			for (const FTriangleID TriangleID : VertexInstanceArray[VertexInstanceID].ConnectedTriangles)
+			{
+				OutConnectedPolygonIDs.AddUnique(TriangleArray[TriangleID].PolygonID);
+			}
 		}
 	}
 
+	/** Returns the polygons connected to this vertex as an array with the specified allocator template type. */
+	template <typename Alloc>
+	TArray<FPolygonID, Alloc> GetVertexConnectedPolygons(const FVertexID VertexID) const
+	{
+		TArray<FPolygonID, Alloc> Result;
+		this->GetVertexConnectedPolygons(VertexID, Result);
+		return Result;
+	}
+
+	/** Returns the polygons connected to this vertex */
+	TArray<FPolygonID> GetVertexConnectedPolygons(const FVertexID VertexID) const
+	{
+		TArray<FPolygonID> Result;
+		this->GetVertexConnectedPolygons(VertexID, Result);
+		return Result;
+	}
+
+	/** Returns the number of polygons connected to this vertex */
+	int32 GetNumVertexConnectedPolygons(const FVertexID VertexID) const
+	{
+		return GetVertexConnectedPolygons<TInlineAllocator<8>>(VertexID).Num();
+	}
+
 	/** Populates the passed array of VertexIDs with the vertices adjacent to this vertex */
-	void GetVertexAdjacentVertices( const FVertexID VertexID, TArray<FVertexID>& OutAdjacentVertexIDs ) const
+	template <typename Alloc>
+	void GetVertexAdjacentVertices(const FVertexID VertexID, TArray<FVertexID, Alloc>& OutAdjacentVertexIDs) const
 	{
-		const TArray<FEdgeID>& ConnectedEdgeIDs = VertexArray[ VertexID ].ConnectedEdgeIDs;
-		OutAdjacentVertexIDs.SetNumUninitialized( ConnectedEdgeIDs.Num() );
+		const TArray<FEdgeID>& ConnectedEdgeIDs = VertexArray[VertexID].ConnectedEdgeIDs;
+		OutAdjacentVertexIDs.SetNumUninitialized(ConnectedEdgeIDs.Num());
 
 		int32 Index = 0;
-		for( const FEdgeID EdgeID : ConnectedEdgeIDs )
+		for (const FEdgeID EdgeID : ConnectedEdgeIDs)
 		{
-			const FMeshEdge& Edge = EdgeArray[ EdgeID ];
-			OutAdjacentVertexIDs[ Index ] = ( Edge.VertexIDs[ 0 ] == VertexID ) ? Edge.VertexIDs[ 1 ] : Edge.VertexIDs[ 0 ];
+			const FMeshEdge& Edge = EdgeArray[EdgeID];
+			OutAdjacentVertexIDs[Index] = (Edge.VertexIDs[0] == VertexID) ? Edge.VertexIDs[1] : Edge.VertexIDs[0];
 			Index++;
 		}
 	}
 
-	/** Returns reference to an array of polygon IDs connected to this edge */
-	const TArray<FPolygonID>& GetEdgeConnectedPolygons( const FEdgeID EdgeID ) const
+	/** Returns the vertices adjacent to this vertex as an array with the specified allocator template type. */
+	template <typename Alloc>
+	TArray<FVertexID, Alloc> GetVertexAdjacentVertices(const FVertexID VertexID) const
+	{
+		TArray<FVertexID, Alloc> Result;
+		this->GetVertexAdjacentVertices(VertexID, Result);
+		return Result;
+	}
+
+	/** Returns the vertices adjacent to this vertex */
+	TArray<FVertexID> GetVertexAdjacentVertices(const FVertexID VertexID) const
+	{
+		TArray<FVertexID> Result;
+		this->GetVertexAdjacentVertices(VertexID, Result);
+		return Result;
+	}
+
+
+	//////////////////////////////////////////////////////////////////////
+	// Vertex instance operations
+
+	/** Returns the vertex ID associated with the given vertex instance */
+	FVertexID GetVertexInstanceVertex(const FVertexInstanceID VertexInstanceID) const
+	{
+		return VertexInstanceArray[VertexInstanceID].VertexID;
+	}
+
+	/** Returns the edge ID defined by the two given vertex instance IDs, if there is one; otherwise FEdgeID::Invalid */
+	FEdgeID GetVertexInstancePairEdge(const FVertexInstanceID VertexInstanceID0, const FVertexInstanceID VertexInstanceID1) const;
+
+	/** Returns reference to an array of Triangle IDs connected to this vertex instance */
+	const TArray<FTriangleID>& GetVertexInstanceConnectedTriangles(const FVertexInstanceID VertexInstanceID) const
+	{
+		return VertexInstanceArray[VertexInstanceID].ConnectedTriangles;
+	}
+
+	/** Returns the number of triangles connected to this vertex instance */
+	int32 GetNumVertexInstanceConnectedTriangles(const FVertexInstanceID VertexInstanceID) const
+	{
+		return VertexInstanceArray[VertexInstanceID].ConnectedTriangles.Num();
+	}
+
+	/** Populates the passed array with the polygons connected to this vertex instance */
+	template <typename Alloc>
+	void GetVertexInstanceConnectedPolygons(const FVertexInstanceID VertexInstanceID, TArray<FPolygonID, Alloc>& OutPolygonIDs) const
+	{
+		OutPolygonIDs.Reset(VertexInstanceArray[VertexInstanceID].ConnectedTriangles.Num());
+		for (const FTriangleID TriangleID : VertexInstanceArray[VertexInstanceID].ConnectedTriangles)
+		{
+			OutPolygonIDs.AddUnique(TriangleArray[TriangleID].PolygonID);
+		}
+	}
+
+	/** Returns the polygons connected to this vertex instance as an array with the specified allocator template type. */
+	template <typename Alloc>
+	TArray<FPolygonID, Alloc> GetVertexInstanceConnectedPolygons(const FVertexInstanceID VertexInstanceID) const
+	{
+		TArray<FPolygonID, Alloc> Result;
+		this->GetVertexInstanceConnectedPolygons(VertexInstanceID, Result);
+		return Result;
+	}
+
+	/** Returns the polygons connected to this vertex instance */
+	TArray<FPolygonID> GetVertexInstanceConnectedPolygons(const FVertexInstanceID VertexInstanceID) const
 	{
-		return EdgeArray[ EdgeID ].ConnectedPolygons;
+		TArray<FPolygonID> Result;
+		this->GetVertexInstanceConnectedPolygons(VertexInstanceID, Result);
+		return Result;
+	}
+
+	/** Returns the number of polygons connected to this vertex instance. */
+	int32 GetNumVertexInstanceConnectedPolygons(const FVertexInstanceID VertexInstanceID) const
+	{
+		return GetVertexInstanceConnectedPolygons<TInlineAllocator<8>>(VertexInstanceID).Num();
+	}
+
+
+	//////////////////////////////////////////////////////////////////////
+	// Edge operations
+
+	/** Determine whether a given edge is an internal edge between triangles of a polygon */
+	bool IsEdgeInternal(const FEdgeID EdgeID) const
+	{
+		const TArray<FTriangleID>& ConnectedTriangles = EdgeArray[EdgeID].ConnectedTriangles;
+		return ConnectedTriangles.Num() == 2 &&
+			   TriangleArray[ConnectedTriangles[0]].PolygonID == TriangleArray[ConnectedTriangles[1]].PolygonID;
+	}
+
+	/** Determine whether a given edge is an internal edge between triangles of a specific polygon */
+	bool IsEdgeInternalToPolygon(const FEdgeID EdgeID, const FPolygonID PolygonID) const
+	{
+		const TArray<FTriangleID>& ConnectedTriangles = EdgeArray[EdgeID].ConnectedTriangles;
+		return ConnectedTriangles.Num() == 2 &&
+			   TriangleArray[ConnectedTriangles[0]].PolygonID == PolygonID &&
+			   TriangleArray[ConnectedTriangles[1]].PolygonID == PolygonID;
+	}
+
+	/** Returns reference to an array of triangle IDs connected to this edge */
+	const TArray<FTriangleID>& GetEdgeConnectedTriangles(const FEdgeID EdgeID) const
+	{
+		return EdgeArray[EdgeID].ConnectedTriangles;
+	}
+
+	int32 GetNumEdgeConnectedTriangles(const FEdgeID EdgeID) const
+	{
+		return EdgeArray[EdgeID].ConnectedTriangles.Num();
+	}
+
+	/** Populates the passed array with polygon IDs connected to this edge */
+	template <typename Alloc>
+	void GetEdgeConnectedPolygons(const FEdgeID EdgeID, TArray<FPolygonID, Alloc>& OutPolygonIDs) const
+	{
+		OutPolygonIDs.Reset(EdgeArray[EdgeID].ConnectedTriangles.Num());
+		for (const FTriangleID TriangleID : EdgeArray[EdgeID].ConnectedTriangles)
+		{
+			OutPolygonIDs.AddUnique(TriangleArray[TriangleID].PolygonID);
+		}
+	}
+
+	/** Returns the polygons connected to this edge as an array with the specified allocator template type. */
+	template <typename Alloc>
+	TArray<FPolygonID, Alloc> GetEdgeConnectedPolygons(const FEdgeID EdgeID) const
+	{
+		TArray<FPolygonID, Alloc> Result;
+		this->GetEdgeConnectedPolygons(EdgeID, Result);
+		return Result;
+	}
+
+	/** Returns the polygons connected to this edge */
+	TArray<FPolygonID> GetEdgeConnectedPolygons(const FEdgeID EdgeID) const
+	{
+		TArray<FPolygonID> Result;
+		this->GetEdgeConnectedPolygons(EdgeID, Result);
+		return Result;
+	}
+
+	/** Returns the number of polygons connected to this edge */
+	int32 GetNumEdgeConnectedPolygons(const FEdgeID EdgeID) const
+	{
+		return GetEdgeConnectedPolygons<TInlineAllocator<8>>(EdgeID).Num();
 	}
 
 	/** Returns the vertex ID corresponding to one of the edge endpoints */
-	FVertexID GetEdgeVertex( const FEdgeID EdgeID, int32 VertexNumber ) const
+	FVertexID GetEdgeVertex(const FEdgeID EdgeID, int32 VertexNumber) const
 	{
-		check( VertexNumber == 0 || VertexNumber == 1 );
-		return EdgeArray[ EdgeID ].VertexIDs[ VertexNumber ];
+		check(VertexNumber == 0 || VertexNumber == 1);
+		return EdgeArray[EdgeID].VertexIDs[VertexNumber];
 	}
 
 	/** Returns a pair of vertex IDs defining the edge */
-	TTuple<FVertexID, FVertexID> GetEdgeVertices( const FEdgeID EdgeID ) const
+	TArrayView<const FVertexID> GetEdgeVertices(const FEdgeID EdgeID) const
+	{
+		return EdgeArray[EdgeID].VertexIDs;
+	}
+
+
+	//////////////////////////////////////////////////////////////////////
+	// Triangle operations
+
+	/** Get the polygon which contains this triangle */
+	FPolygonID GetTrianglePolygon(const FTriangleID TriangleID) const
+	{
+		return TriangleArray[TriangleID].PolygonID;
+	}
+
+	/** Get the polygon group which contains this triangle */
+	FPolygonGroupID GetTrianglePolygonGroup(const FTriangleID TriangleID) const
+	{
+		return PolygonArray[TriangleArray[TriangleID].PolygonID].PolygonGroupID;
+	}
+
+	/** Determines if this triangle is part of an n-gon */
+	bool IsTrianglePartOfNgon(const FTriangleID TriangleID) const
+	{
+		return PolygonArray[TriangleArray[TriangleID].PolygonID].TriangleIDs.Num() > 1;
+	}
+
+	/** Get the vertex instances which define this triangle */
+	TArrayView<const FVertexInstanceID> GetTriangleVertexInstances(const FTriangleID TriangleID) const
+	{
+		return TArrayView<const FVertexInstanceID>(&TriangleArray[TriangleID].VertexInstanceID0, 3);
+	}
+
+	/** Get the specified vertex instance by index */
+	FVertexInstanceID GetTriangleVertexInstance(const FTriangleID TriangleID, const int32 Index) const
+	{
+		check(Index >= 0 && Index < 3);
+		return TriangleArray[TriangleID].GetVertexInstanceID(Index);
+	}
+
+	/** Populates the passed array with the vertices which define this triangle */
+	void GetTriangleVertices(const FTriangleID TriangleID, TArrayView<FVertexID> OutVertexIDs) const
+	{
+		check(OutVertexIDs.Num() >= 3);
+		for (int32 Index = 0; Index < 3; ++Index)
+		{
+			OutVertexIDs[Index] = GetVertexInstanceVertex(TriangleArray[TriangleID].GetVertexInstanceID(Index));
+		}
+	}
+
+	/** Return the vertices which define this triangle */
+	TStaticArray<FVertexID, 3> GetTriangleVertices(const FTriangleID TriangleID) const
+	{
+		TStaticArray<FVertexID, 3> Result;
+		GetTriangleVertices(TriangleID, Result);
+		return Result;
+	}
+
+	/** Populates the passed array with the edges which define this triangle */
+	void GetTriangleEdges(const FTriangleID TriangleID, TArrayView<FEdgeID> OutEdgeIDs) const
+	{
+		check(OutEdgeIDs.Num() >= 3);
+		TStaticArray<FVertexID, 3> VertexIDs = GetTriangleVertices(TriangleID);
+		OutEdgeIDs[0] = GetVertexPairEdge(VertexIDs[0], VertexIDs[1]);
+		OutEdgeIDs[1] = GetVertexPairEdge(VertexIDs[1], VertexIDs[2]);
+		OutEdgeIDs[2] = GetVertexPairEdge(VertexIDs[2], VertexIDs[0]);
+	}
+
+	/** Return the edges which form this triangle */
+	TStaticArray<FEdgeID, 3> GetTriangleEdges(const FTriangleID TriangleID) const
+	{
+		TStaticArray<FEdgeID, 3> Result;
+		GetTriangleEdges(TriangleID, Result);
+		return Result;
+	}
+
+	/** Populates the passed array with adjacent triangles */
+	template <typename Alloc>
+	void GetTriangleAdjacentTriangles(const FTriangleID TriangleID, TArray<FTriangleID, Alloc>& OutTriangleIDs) const
+	{
+		OutTriangleIDs.Reset();
+		for (const FEdgeID EdgeID : GetTriangleEdges(TriangleID))
+		{
+			for (const FTriangleID OtherTriangleID : EdgeArray[EdgeID].ConnectedTriangles)
+			{
+				if (OtherTriangleID != TriangleID)
+				{
+					OutTriangleIDs.Add(OtherTriangleID);
+				}
+			}
+		}
+	}
+
+	/** Return adjacent triangles into a TArray with the specified allocator */
+	template <typename Alloc>
+	TArray<FTriangleID, Alloc> GetTriangleAdjacentTriangles(const FTriangleID TriangleID) const
 	{
-		const FMeshEdge& Edge = EdgeArray[ EdgeID ];
-		return MakeTuple( Edge.VertexIDs[ 0 ], Edge.VertexIDs[ 1 ] );
+		TArray<FTriangleID, Alloc> Result;
+		this->GetTriangleAdjacentTriangles(TriangleID, Result);
+		return Result;
 	}
 
-	/** Returns reference to the array of triangles representing the triangulated polygon */
-	TArray<FMeshTriangle>& GetPolygonTriangles( const FPolygonID PolygonID )
+	/** Return adjacent triangles to this triangle */
+	TArray<FTriangleID> GetTriangleAdjacentTriangles(const FTriangleID TriangleID) const
 	{
-		return PolygonArray[ PolygonID ].Triangles;
+		TArray<FTriangleID> Result;
+		this->GetTriangleAdjacentTriangles(TriangleID, Result);
+		return Result;
 	}
 
-	const TArray<FMeshTriangle>& GetPolygonTriangles( const FPolygonID PolygonID ) const
+	/** Return the vertex instance which corresponds to the given vertex on the given triangle, or FVertexInstanceID::Invalid */
+	FVertexInstanceID GetVertexInstanceForTriangleVertex(const FTriangleID TriangleID, const FVertexID VertexID) const
 	{
-		return PolygonArray[ PolygonID ].Triangles;
+		const FVertexInstanceID* VertexInstanceIDPtr = Algo::FindByPredicate(
+			GetTriangleVertexInstances(TriangleID),
+			[this, VertexID](const FVertexInstanceID VertexInstanceID) { return (GetVertexInstanceVertex(VertexInstanceID) == VertexID); });
+
+		return VertexInstanceIDPtr ? *VertexInstanceIDPtr : FVertexInstanceID::Invalid;
+	}
+
+
+	//////////////////////////////////////////////////////////////////////
+	// Polygon operations
+
+	UE_DEPRECATED(4.24, "Please use GetPolygonTriangleIDs instead.")
+	TArray<FMeshTriangle> GetPolygonTriangles(const FPolygonID PolygonID) const;
+
+	/** Return reference to an array of triangle IDs which comprise this polygon */
+	const TArray<FTriangleID>& GetPolygonTriangleIDs(const FPolygonID PolygonID) const
+	{
+		return PolygonArray[PolygonID].TriangleIDs;
+	}
+
+	/** Return the number of triangles which comprise this polygon */
+	int32 GetNumPolygonTriangles(const FPolygonID PolygonID) const
+	{
+		return PolygonArray[PolygonID].TriangleIDs.Num();
+	}
+
+	UE_DEPRECATED(4.24, "Please use GetPolygonVertexInstances instead.")
+	const TArray<FVertexInstanceID>& GetPolygonPerimeterVertexInstances(const FPolygonID PolygonID) const
+	{
+		return GetPolygonVertexInstances(PolygonID);
 	}
 
 	/** Returns reference to an array of VertexInstance IDs forming the perimeter of this polygon */
-	const TArray<FVertexInstanceID>& GetPolygonPerimeterVertexInstances( const FPolygonID PolygonID ) const
+	const TArray<FVertexInstanceID>& GetPolygonVertexInstances(const FPolygonID PolygonID) const
 	{
-		return PolygonArray[ PolygonID ].PerimeterContour.VertexInstanceIDs;
+		return PolygonArray[PolygonID].PerimeterContour.VertexInstanceIDs;
+	}
+
+	/** Returns the number of vertices this polygon has */
+	int32 GetNumPolygonVertices(const FPolygonID PolygonID) const
+	{
+		return PolygonArray[PolygonID].PerimeterContour.VertexInstanceIDs.Num();
+	}
+
+	UE_DEPRECATED(4.24, "Please use GetPolygonVertices instead.")
+	void GetPolygonPerimeterVertices(const FPolygonID PolygonID, TArray<FVertexID>& OutVertexIDs) const
+	{
+		return GetPolygonVertices(PolygonID, OutVertexIDs);
 	}
 
 	/** Populates the passed array of VertexIDs with the vertices which form the polygon perimeter */
-	void GetPolygonPerimeterVertices( const FPolygonID PolygonID, TArray<FVertexID>& OutPolygonPerimeterVertexIDs ) const;
+	template <typename Alloc>
+	void GetPolygonVertices(const FPolygonID PolygonID, TArray<FVertexID, Alloc>& OutVertexIDs) const
+	{
+		OutVertexIDs.SetNumUninitialized(GetNumPolygonVertices(PolygonID));
+		int32 Index = 0;
+		for (const FVertexInstanceID VertexInstanceID : GetPolygonVertexInstances(PolygonID))
+		{
+			OutVertexIDs[Index++] = GetVertexInstanceVertex(VertexInstanceID);
+		}
+	}
 
-	void GetPolygonEdges(const FPolygonID PolygonID, TArray<FEdgeID>& OutPolygonEdges) const
+	/** Returns the vertices which form the polygon perimeter as an array templated on the given allocator */
+	template <typename Alloc>
+	TArray<FVertexID, Alloc> GetPolygonVertices(const FPolygonID PolygonID) const
 	{
-		const FMeshPolygonContour& PerimeterContour = PolygonArray[PolygonID].PerimeterContour;
-		const int32 ContourCount = PerimeterContour.VertexInstanceIDs.Num();
+		TArray<FVertexID, Alloc> Result;
+		this->GetPolygonVertices(PolygonID, Result);
+		return Result;
+	}
+
+	/** Returns the vertices which form the polygon perimeter */
+	TArray<FVertexID> GetPolygonVertices(const FPolygonID PolygonID) const
+	{
+		TArray<FVertexID> Result;
+		this->GetPolygonVertices(PolygonID, Result);
+		return Result;
+	}
+
+	UE_DEPRECATED(4.24, "Please use GetPolygonPerimeterEdges instead.")
+	void GetPolygonEdges(const FPolygonID PolygonID, TArray<FEdgeID>& OutEdgeIDs) const
+	{
+		GetPolygonPerimeterEdges(PolygonID, OutEdgeIDs);
+	}
+
+	/** Populates the passed array with the edges which form the polygon perimeter */
+	template <typename Alloc>
+	void GetPolygonPerimeterEdges(const FPolygonID PolygonID, TArray<FEdgeID, Alloc>& OutEdgeIDs) const
+	{
+		const TArray<FVertexInstanceID>& VertexInstanceIDs = GetPolygonVertexInstances(PolygonID);
+		const int32 ContourCount = VertexInstanceIDs.Num();
+		OutEdgeIDs.SetNumUninitialized(ContourCount);
 		for (int32 ContourIndex = 0; ContourIndex < ContourCount; ++ContourIndex)
 		{
-			int32 ContourPlusOne = (ContourIndex + 1) % ContourCount;
-			OutPolygonEdges.Add(GetVertexPairEdge(GetVertexInstanceVertex(PerimeterContour.VertexInstanceIDs[ContourIndex]), GetVertexInstanceVertex(PerimeterContour.VertexInstanceIDs[ContourPlusOne])));
+			const int32 ContourPlusOne = (ContourIndex == ContourCount - 1) ? 0 : ContourIndex + 1;
+			OutEdgeIDs[ContourIndex] = GetVertexPairEdge(GetVertexInstanceVertex(VertexInstanceIDs[ContourIndex]),
+														 GetVertexInstanceVertex(VertexInstanceIDs[ContourPlusOne]));
 		}
 	}
 
-	/** Return the polygon group associated with a polygon */
-	FPolygonGroupID GetPolygonPolygonGroup( const FPolygonID PolygonID ) const
+	/** Returns the vertices which form the polygon perimeter as an array templated on the given allocator */
+	template <typename Alloc>
+	TArray<FEdgeID, Alloc> GetPolygonPerimeterEdges(const FPolygonID PolygonID) const
 	{
-		return PolygonArray[ PolygonID ].PolygonGroupID;
+		TArray<FEdgeID, Alloc> Result;
+		this->GetPolygonPerimeterEdges(PolygonID, Result);
+		return Result;
 	}
 
-	/** Sets the polygon group associated with a polygon */
-	void SetPolygonPolygonGroup( const FPolygonID PolygonID, const FPolygonGroupID PolygonGroupID )
+	/** Returns the vertices which form the polygon perimeter */
+	TArray<FEdgeID> GetPolygonPerimeterEdges(const FPolygonID PolygonID) const
 	{
-		FMeshPolygon& Polygon = PolygonArray[ PolygonID ];
-		verify( PolygonGroupArray[ Polygon.PolygonGroupID ].Polygons.Remove( PolygonID ) == 1 );
-		Polygon.PolygonGroupID = PolygonGroupID;
-		check( !PolygonGroupArray[ PolygonGroupID ].Polygons.Contains( PolygonID ) );
-		PolygonGroupArray[ PolygonGroupID ].Polygons.Add( PolygonID );
+		TArray<FEdgeID> Result;
+		this->GetPolygonPerimeterEdges(PolygonID, Result);
+		return Result;
 	}
 
-	/** Return the vertex instance which corresponds to the given vertex on the given polygon, or FVertexInstanceID::Invalid */
-	FVertexInstanceID GetVertexInstanceForPolygonVertex( const FPolygonID PolygonID, const FVertexID VertexID ) const
+	/** Populate the provided array with a list of edges which are internal to the polygon, i.e. those which separate
+	    constituent triangles. */
+	template <typename Alloc>
+	void GetPolygonInternalEdges(const FPolygonID PolygonID, TArray<FEdgeID, Alloc>& OutEdgeIDs) const
 	{
-		const FVertexInstanceID* VertexInstanceIDPtr = PolygonArray[ PolygonID ].PerimeterContour.VertexInstanceIDs.FindByPredicate(
-			[ this, VertexID ]( const FVertexInstanceID VertexInstanceID )
+		OutEdgeIDs.Reset(GetNumPolygonVertices(PolygonID) - 3);
+		if (GetNumPolygonVertices(PolygonID) > 3)
+		{
+			for (const FVertexInstanceID VertexInstanceID : GetPolygonVertexInstances(PolygonID))
 			{
-				return GetVertexInstanceVertex( VertexInstanceID ) == VertexID;
-			});
+				for (const FEdgeID EdgeID : GetVertexConnectedEdges(GetVertexInstanceVertex(VertexInstanceID)))
+				{
+					if (!OutEdgeIDs.Contains(EdgeID) && IsEdgeInternalToPolygon(EdgeID, PolygonID))
+					{
+						OutEdgeIDs.Add(EdgeID);
+					}
+				}
+			}
+		}
+	}
 
-		return VertexInstanceIDPtr ? *VertexInstanceIDPtr : FVertexInstanceID::Invalid;
+	/** Return the internal edges of this polygon, i.e. those which separate constituent triangles */
+	template <typename Alloc>
+	TArray<FEdgeID, Alloc> GetPolygonInternalEdges(const FPolygonID PolygonID) const
+	{
+		TArray<FEdgeID, Alloc> Result;
+		this->GetPolygonInternalEdges(PolygonID, Result);
+		return Result;
 	}
 
-	/** Returns the vertex ID associated with the given vertex instance */
-	FVertexID GetVertexInstanceVertex( const FVertexInstanceID VertexInstanceID ) const
+	/** Return the internal edges of this polygon, i.e. those which separate constituent triangles */
+	TArray<FEdgeID> GetPolygonInternalEdges(const FPolygonID PolygonID) const
+	{
+		TArray<FEdgeID> Result;
+		this->GetPolygonInternalEdges(PolygonID, Result);
+		return Result;
+	}
+
+	/** Return the number of internal edges in this polygon */
+	int32 GetNumPolygonInternalEdges(const FPolygonID PolygonID) const
+	{
+		return PolygonArray[PolygonID].PerimeterContour.VertexInstanceIDs.Num() - 3;
+	}
+
+	/** Populates the passed array with adjacent polygons */
+	template <typename Alloc>
+	void GetPolygonAdjacentPolygons(const FPolygonID PolygonID, TArray<FPolygonID, Alloc>& OutPolygonIDs) const
+	{
+		OutPolygonIDs.Reset();
+		for (const FEdgeID EdgeID : GetPolygonPerimeterEdges<TInlineAllocator<16>>(PolygonID))
+		{
+			for (const FPolygonID OtherPolygonID : GetEdgeConnectedPolygons<TInlineAllocator<8>>(EdgeID))
+			{
+				if (OtherPolygonID != PolygonID)
+				{
+					OutPolygonIDs.Add(OtherPolygonID);
+				}
+			}
+		}
+	}
+
+	/** Return adjacent polygons into a TArray with the specified allocator */
+	template <typename Alloc>
+	TArray<FPolygonID, Alloc> GetPolygonAdjacentPolygons(const FPolygonID PolygonID) const
 	{
-		return VertexInstanceArray[ VertexInstanceID ].VertexID;
+		TArray<FPolygonID, Alloc> Result;
+		this->GetPolygonAdjacentPolygons(PolygonID, Result);
+		return Result;
 	}
 
-	/** Returns reference to an array of Polygon IDs connected to this vertex instance */
-	const TArray<FPolygonID>& GetVertexInstanceConnectedPolygons( const FVertexInstanceID VertexInstanceID ) const
+	/** Return adjacent polygons to this polygon */
+	TArray<FPolygonID> GetPolygonAdjacentPolygons(const FPolygonID PolygonID) const
 	{
-		return VertexInstanceArray[ VertexInstanceID ].ConnectedPolygons;
+		TArray<FPolygonID> Result;
+		this->GetPolygonAdjacentPolygons(PolygonID, Result);
+		return Result;
 	}
 
+	/** Return the polygon group associated with a polygon */
+	FPolygonGroupID GetPolygonPolygonGroup(const FPolygonID PolygonID) const
+	{
+		return PolygonArray[PolygonID].PolygonGroupID;
+	}
+
+	/** Return the vertex instance which corresponds to the given vertex on the given polygon, or FVertexInstanceID::Invalid */
+	FVertexInstanceID GetVertexInstanceForPolygonVertex(const FPolygonID PolygonID, const FVertexID VertexID) const
+	{
+		const FVertexInstanceID* VertexInstanceIDPtr = GetPolygonVertexInstances(PolygonID).FindByPredicate(
+			[this, VertexID](const FVertexInstanceID VertexInstanceID) { return (GetVertexInstanceVertex(VertexInstanceID) == VertexID); });
+
+		return VertexInstanceIDPtr ? *VertexInstanceIDPtr : FVertexInstanceID::Invalid;
+	}
+
+	/** Set the vertex instance at the given index around the polygon to the new value */
+	void SetPolygonVertexInstance(const FPolygonID PolygonID, const int32 PerimeterIndex, const FVertexInstanceID VertexInstanceID);
+
+	/** Sets the polygon group associated with a polygon */
+	void SetPolygonPolygonGroup(const FPolygonID PolygonID, const FPolygonGroupID PolygonGroupID)
+	{
+		FMeshPolygon& Polygon = PolygonArray[PolygonID];
+		verify(PolygonGroupArray[Polygon.PolygonGroupID].Polygons.Remove(PolygonID) == 1);
+		Polygon.PolygonGroupID = PolygonGroupID;
+		check(!PolygonGroupArray[PolygonGroupID].Polygons.Contains(PolygonID));
+		PolygonGroupArray[PolygonGroupID].Polygons.Add(PolygonID);
+	}
+
+	/** Reverse the winding order of the vertices of this polygon */
+	void ReversePolygonFacing(const FPolygonID PolygonID);
+
+	UE_DEPRECATED(4.24, "Please use the other overload of ComputePolygonTriangulation.")
+	void ComputePolygonTriangulation(const FPolygonID PolygonID, TArray<FMeshTriangle>& OutTriangles);
+
+	/** Generates triangles and internal edges for the given polygon */
+	void ComputePolygonTriangulation(const FPolygonID PolygonID);
+
+
+	//////////////////////////////////////////////////////////////////////
+	// Polygon group operations
+
 	/** Returns the polygons associated with the given polygon group */
-	const TArray<FPolygonID>& GetPolygonGroupPolygons( const FPolygonGroupID PolygonGroupID ) const
+	const TArray<FPolygonID>& GetPolygonGroupPolygons(const FPolygonGroupID PolygonGroupID) const
+	{
+		return PolygonGroupArray[PolygonGroupID].Polygons;
+	}
+
+	/** Returns the number of polygons in this polygon group */
+	int32 GetNumPolygonGroupPolygons(const FPolygonGroupID PolygonGroupID) const
 	{
-		return PolygonGroupArray[ PolygonGroupID ].Polygons;
+		return PolygonGroupArray[PolygonGroupID].Polygons.Num();
 	}
 
+	/** Remaps polygon groups according to the supplied map */
+	void RemapPolygonGroups(const TMap<FPolygonGroupID, FPolygonGroupID>& Remap);
+
+
+	//////////////////////////////////////////////////////////////////////
+	// Whole mesh operations
+
 	/** Compacts the data held in the mesh description, and returns an object describing how the IDs have been remapped. */
 	void Compact( FElementIDRemappings& OutRemappings );
 
 	/** Remaps the element IDs in the mesh description according to the passed in object */
 	void Remap( const FElementIDRemappings& Remappings );
 
-
-	void ComputePolygonTriangulation(const FPolygonID PolygonID, TArray<FMeshTriangle>& OutTriangles);
 	void TriangulateMesh();
 
 	/** Set the polygon tangent and normal only for the specified polygonIDs */
 	void ComputePolygonTangentsAndNormals(const TArray<FPolygonID>& PolygonIDs, float ComparisonThreshold = 0.0f);
+
 	/** Set the polygon tangent and normal for all polygons in the mesh description. */
 	void ComputePolygonTangentsAndNormals(float ComparisonThreshold = 0.0f);
 	
 	/** Set the vertex instance tangent and normal only for the specified VertexInstanceIDs */
 	void ComputeTangentsAndNormals(const TArray<FVertexInstanceID>& VertexInstanceIDs, EComputeNTBsOptions ComputeNTBsOptions);
+
 	/** Set the vertex instance tangent and normal for all vertex instances in the mesh description. */
 	void ComputeTangentsAndNormals(EComputeNTBsOptions ComputeNTBsOptions);
 
@@ -829,16 +1312,15 @@ public:
 	/** Get array of all UV charts */
 	void GetAllCharts( TArray<TArray<FPolygonID>>& OutCharts );
 
-	void ReversePolygonFacing(const FPolygonID PolygonID);
 	void ReverseAllPolygonFacing();
 
+	float GetPolygonCornerAngleForVertex(const FPolygonID PolygonID, const FVertexID VertexID) const;
+
 private:
-	bool VectorsOnSameSide(const FVector& Vec, const FVector& A, const FVector& B, const float SameSideDotProductEpsilon);
-	bool PointInTriangle(const FVector& A, const FVector& B, const FVector& C, const FVector& P, const float InsideTriangleDotProductEpsilon);
+
 	FPlane ComputePolygonPlane(const FPolygonID PolygonID) const;
 	FVector ComputePolygonNormal(const FPolygonID PolygonID) const;
 
-	float GetPolygonCornerAngleForVertex(const FPolygonID PolygonID, const FVertexID VertexID) const;
 	bool ComputePolygonTangentsAndNormals(
 		  const FPolygonID PolygonID
 		, float ComparisonThreshold
@@ -863,8 +1345,16 @@ private:
 		, TVertexInstanceAttributesRef<FVector> VertexTangents
 		, TVertexInstanceAttributesRef<float> VertexBinormalSigns
 	);
+
 private:
 
+	void CreateVertex_Internal( const FVertexID VertexID ) { VertexAttributesSet.Insert( VertexID ); }
+	void CreateVertexInstance_Internal( const FVertexInstanceID VertexInstanceID, const FVertexID VertexID );
+	void CreateEdge_Internal( const FEdgeID EdgeID, const FVertexID VertexID0, const FVertexID VertexID1 );
+	void CreateTriangle_Internal( const FTriangleID TriangleID, const FPolygonGroupID PolygonGroupID, TArrayView<const FVertexInstanceID> VertexInstanceIDs, TArray<FEdgeID>* OutEdgeIDs );
+	void CreatePolygon_Internal( const FPolygonID PolygonID, const FPolygonGroupID PolygonGroupID, TArrayView<const FVertexInstanceID> VertexInstanceIDs, TArray<FEdgeID>* OutEdgeIDs );
+	void CreatePolygonGroup_Internal( const FPolygonGroupID PolygonGroupID ) { PolygonGroupAttributesSet.Insert( PolygonGroupID ); }
+
 	/** Given a set of index remappings, fixes up references to element IDs */
 	void FixUpElementIDs( const FElementIDRemappings& Remappings );
 
@@ -875,16 +1365,20 @@ private:
 	FVertexArray VertexArray;
 	FVertexInstanceArray VertexInstanceArray;
 	FEdgeArray EdgeArray;
+	FTriangleArray TriangleArray;
 	FPolygonArray PolygonArray;
 	FPolygonGroupArray PolygonGroupArray;
 
 	TAttributesSet<FVertexID> VertexAttributesSet;
 	TAttributesSet<FVertexInstanceID> VertexInstanceAttributesSet;
 	TAttributesSet<FEdgeID> EdgeAttributesSet;
+	TAttributesSet<FTriangleID> TriangleAttributesSet;
 	TAttributesSet<FPolygonID> PolygonAttributesSet;
 	TAttributesSet<FPolygonGroupID> PolygonGroupAttributesSet;
 };
 
+PRAGMA_ENABLE_DEPRECATION_WARNINGS
+
 
 /**
  * Bulk data storage for FMeshDescription
diff --git a/Engine/Source/Runtime/MeshDescription/Public/MeshElementRemappings.h b/Engine/Source/Runtime/MeshDescription/Public/MeshElementRemappings.h
--- a/Engine/Source/Runtime/MeshDescription/Public/MeshElementRemappings.h
+++ b/Engine/Source/Runtime/MeshDescription/Public/MeshElementRemappings.h
@@ -13,6 +13,7 @@ struct FElementIDRemappings
 	TSparseArray<int32> NewVertexIndexLookup;
 	TSparseArray<int32> NewVertexInstanceIndexLookup;
 	TSparseArray<int32> NewEdgeIndexLookup;
+	TSparseArray<int32> NewTriangleIndexLookup;
 	TSparseArray<int32> NewPolygonIndexLookup;
 	TSparseArray<int32> NewPolygonGroupIndexLookup;
 
@@ -34,6 +35,12 @@ struct FElementIDRemappings
 		return FEdgeID( NewEdgeIndexLookup[ EdgeID.GetValue() ] );
 	}
 
+	FTriangleID GetRemappedTriangleID( FTriangleID TriangleID ) const
+	{
+		checkSlow( NewTriangleIndexLookup.IsAllocated( TriangleID.GetValue() ) );
+		return FTriangleID( NewTriangleIndexLookup[ TriangleID.GetValue() ] );
+	}
+
 	FPolygonID GetRemappedPolygonID( FPolygonID PolygonID ) const
 	{
 		checkSlow( NewPolygonIndexLookup.IsAllocated( PolygonID.GetValue() ) );
diff --git a/Engine/Source/Runtime/MeshDescription/Public/MeshTypes.h b/Engine/Source/Runtime/MeshDescription/Public/MeshTypes.h
--- a/Engine/Source/Runtime/MeshDescription/Public/MeshTypes.h
+++ b/Engine/Source/Runtime/MeshDescription/Public/MeshTypes.h
@@ -149,6 +149,35 @@ struct FEdgeID : public FElementID
 };
 
 
+USTRUCT(BlueprintType)
+struct FTriangleID : public FElementID
+{
+	GENERATED_BODY()
+
+	FTriangleID()
+	{
+	}
+
+	explicit FTriangleID(const FElementID InitElementID)
+		: FElementID(InitElementID.GetValue())
+	{
+	}
+
+	explicit FTriangleID(const int32 InitIDValue)
+		: FElementID(InitIDValue)
+	{
+	}
+
+	FORCEINLINE friend uint32 GetTypeHash(const FTriangleID& Other)
+	{
+		return GetTypeHash(Other.IDValue);
+	}
+
+	/** Invalid edge ID */
+	MESHDESCRIPTION_API static const FTriangleID Invalid;
+};
+
+
 USTRUCT( BlueprintType )
 struct FPolygonGroupID : public FElementID
 {
