Removed deprecated members and methods from MeshDescription.

diff --git a/Engine/Source/Runtime/MeshDescription/Private/MeshDescription.cpp b/Engine/Source/Runtime/MeshDescription/Private/MeshDescription.cpp
--- a/Engine/Source/Runtime/MeshDescription/Private/MeshDescription.cpp
+++ b/Engine/Source/Runtime/MeshDescription/Private/MeshDescription.cpp
@@ -8,21 +8,21 @@
 #include "Serialization/BulkDataWriter.h"
 #include "UObject/EnterpriseObjectVersion.h"
 
-void UDEPRECATED_MeshDescription::Serialize( FArchive& Ar )
+void UDEPRECATED_MeshDescription::Serialize(FArchive& Ar)
 {
-	Super::Serialize( Ar );
+	Super::Serialize(Ar);
 
-	if( !HasAnyFlags( RF_ClassDefaultObject ) )
+	if (!HasAnyFlags(RF_ClassDefaultObject))
 	{
-		UE_LOG( LogLoad, Error, TEXT( "UMeshDescription about to be deprecated - please resave %s" ), *GetPathName() );
+		UE_LOG(LogLoad, Error, TEXT( "UMeshDescription about to be deprecated - please resave %s"), *GetPathName());
 	}
+
 	// Discard the contents
 	FMeshDescription MeshDescription;
 	Ar << MeshDescription;
 }
 
 
-PRAGMA_DISABLE_DEPRECATION_WARNINGS
 
 FMeshDescription::FMeshDescription()
 {
@@ -94,33 +94,17 @@ void FMeshDescription::Serialize(FArchive& Ar)
 			if (Ar.CustomVer(FEditorObjectVersion::GUID) >= FEditorObjectVersion::MeshDescriptionTriangles)
 			{
 				// If the polygon has no contour serialized, copy it over from the triangle
-				if (PolygonArray[PolygonID].PerimeterContour.VertexInstanceIDs.Num() == 0)
+				if (PolygonArray[PolygonID].VertexInstanceIDs.Num() == 0)
 				{
 					check(PolygonArray[PolygonID].TriangleIDs.Num() == 1);
 					const FTriangleID TriangleID = PolygonArray[PolygonID].TriangleIDs[0];
 					for (int32 Index = 0; Index < 3; Index++)
 					{
-						PolygonArray[PolygonID].PerimeterContour.VertexInstanceIDs.Add(TriangleArray[TriangleID].GetVertexInstanceID(Index));
+						PolygonArray[PolygonID].VertexInstanceIDs.Add(TriangleArray[TriangleID].GetVertexInstanceID(Index));
 					}
 				}
 			}
 
-			// For backward compatibility only; in 4.24, this will be removed.
-			// Note: internal edges will have the same connected polygon twice.
-			const TArray<FVertexInstanceID>& VertexInstanceIDs = GetPolygonVertexInstances(PolygonID);
-			const int32 NumVertexInstances = VertexInstanceIDs.Num();
-			for (int32 Index = 0; Index < NumVertexInstances; ++Index)
-			{
-				const FVertexInstanceID VertexInstanceID0 = VertexInstanceIDs[Index];
-				const FVertexInstanceID VertexInstanceID1 = VertexInstanceIDs[(Index + 1) % NumVertexInstances];
-				const FVertexID VertexID0 = GetVertexInstanceVertex(VertexInstanceID0);
-				const FVertexID VertexID1 = GetVertexInstanceVertex(VertexInstanceID1);
-				const FEdgeID EdgeID = GetVertexPairEdge(VertexID0, VertexID1);
-				VertexInstanceArray[VertexInstanceID0].ConnectedPolygons.Add(PolygonID);
-				EdgeArray[EdgeID].ConnectedPolygons.Add(PolygonID);
-			}
-			// End backward compatibility
-
 			const FPolygonGroupID PolygonGroupID = PolygonArray[PolygonID].PolygonGroupID;
 			PolygonGroupArray[PolygonGroupID].Polygons.Add(PolygonID);
 		}
@@ -161,24 +145,6 @@ void FMeshDescription::Serialize(FArchive& Ar)
 			}
 		}
 	}
-
-	if (Ar.IsLoading())
-	{
-		bool bAllTriangles = true;
-		int32 TriCount = 0;
-		for (const FPolygonID PolygonID : PolygonArray.GetElementIDs())
-		{
-			int32 NumTriangles = GetPolygonTriangleIDs(PolygonID).Num();
-			TriCount += NumTriangles;
-			if (NumTriangles > 3)
-			{
-				bAllTriangles = false;
-				UE_LOG(LogTemp, Log, TEXT("Polygon %d is a quad (or more)"), PolygonID.GetValue());
-			}
-		}
-
-		check(!bAllTriangles || TriCount == TriangleArray.Num());
-	}
 }
 
 
@@ -191,13 +157,13 @@ void FMeshDescription::Empty()
 	PolygonArray.Reset();
 	PolygonGroupArray.Reset();
 
-	//Empty all attributes
-	VertexAttributesSet.Initialize( 0 );
-	VertexInstanceAttributesSet.Initialize( 0 );
-	EdgeAttributesSet.Initialize( 0 );
-	TriangleAttributesSet.Initialize( 0 );
-	PolygonAttributesSet.Initialize( 0 );
-	PolygonGroupAttributesSet.Initialize( 0 );
+	// Empty all attributes
+	VertexAttributesSet.Initialize(0);
+	VertexInstanceAttributesSet.Initialize(0);
+	EdgeAttributesSet.Initialize(0);
+	TriangleAttributesSet.Initialize(0);
+	PolygonAttributesSet.Initialize(0);
+	PolygonGroupAttributesSet.Initialize(0);
 }
 
 
@@ -212,141 +178,131 @@ bool FMeshDescription::IsEmpty() const
 }
 
 
-void FMeshDescription::Compact( FElementIDRemappings& OutRemappings )
+void FMeshDescription::Compact(FElementIDRemappings& OutRemappings)
 {
-	VertexArray.Compact( OutRemappings.NewVertexIndexLookup );
-	VertexInstanceArray.Compact( OutRemappings.NewVertexInstanceIndexLookup );
-	EdgeArray.Compact( OutRemappings.NewEdgeIndexLookup );
-	TriangleArray.Compact( OutRemappings.NewTriangleIndexLookup );
-	PolygonArray.Compact( OutRemappings.NewPolygonIndexLookup );
-	PolygonGroupArray.Compact( OutRemappings.NewPolygonGroupIndexLookup );
-
-	RemapAttributes( OutRemappings );
-	FixUpElementIDs( OutRemappings );
+	VertexArray.Compact(OutRemappings.NewVertexIndexLookup);
+	VertexInstanceArray.Compact(OutRemappings.NewVertexInstanceIndexLookup);
+	EdgeArray.Compact(OutRemappings.NewEdgeIndexLookup);
+	TriangleArray.Compact(OutRemappings.NewTriangleIndexLookup);
+	PolygonArray.Compact(OutRemappings.NewPolygonIndexLookup);
+	PolygonGroupArray.Compact(OutRemappings.NewPolygonGroupIndexLookup);
+
+	RemapAttributes(OutRemappings);
+	FixUpElementIDs(OutRemappings);
 }
 
 
-void FMeshDescription::Remap( const FElementIDRemappings& Remappings )
+void FMeshDescription::Remap(const FElementIDRemappings& Remappings)
 {
-	VertexArray.Remap( Remappings.NewVertexIndexLookup );
-	VertexInstanceArray.Remap( Remappings.NewVertexInstanceIndexLookup );
-	EdgeArray.Remap( Remappings.NewEdgeIndexLookup );
-	TriangleArray.Remap( Remappings.NewTriangleIndexLookup );
-	PolygonArray.Remap( Remappings.NewPolygonIndexLookup );
-	PolygonGroupArray.Remap( Remappings.NewPolygonGroupIndexLookup );
-
-	RemapAttributes( Remappings );
-	FixUpElementIDs( Remappings );
+	VertexArray.Remap(Remappings.NewVertexIndexLookup);
+	VertexInstanceArray.Remap(Remappings.NewVertexInstanceIndexLookup);
+	EdgeArray.Remap(Remappings.NewEdgeIndexLookup);
+	TriangleArray.Remap(Remappings.NewTriangleIndexLookup);
+	PolygonArray.Remap(Remappings.NewPolygonIndexLookup);
+	PolygonGroupArray.Remap(Remappings.NewPolygonGroupIndexLookup);
+
+	RemapAttributes(Remappings);
+	FixUpElementIDs(Remappings);
 }
 
 
-void FMeshDescription::RemapAttributes( const FElementIDRemappings& Remappings )
+void FMeshDescription::RemapAttributes(const FElementIDRemappings& Remappings)
 {
-	VertexAttributesSet.Remap( Remappings.NewVertexIndexLookup );
-	VertexInstanceAttributesSet.Remap( Remappings.NewVertexInstanceIndexLookup );
-	EdgeAttributesSet.Remap( Remappings.NewEdgeIndexLookup );
-	TriangleAttributesSet.Remap( Remappings.NewTriangleIndexLookup );
-	PolygonAttributesSet.Remap( Remappings.NewPolygonIndexLookup );
-	PolygonGroupAttributesSet.Remap( Remappings.NewPolygonGroupIndexLookup );
+	VertexAttributesSet.Remap(Remappings.NewVertexIndexLookup);
+	VertexInstanceAttributesSet.Remap(Remappings.NewVertexInstanceIndexLookup);
+	EdgeAttributesSet.Remap(Remappings.NewEdgeIndexLookup);
+	TriangleAttributesSet.Remap(Remappings.NewTriangleIndexLookup);
+	PolygonAttributesSet.Remap(Remappings.NewPolygonIndexLookup);
+	PolygonGroupAttributesSet.Remap(Remappings.NewPolygonGroupIndexLookup);
 }
 
 
-void FMeshDescription::FixUpElementIDs( const FElementIDRemappings& Remappings )
+void FMeshDescription::FixUpElementIDs(const FElementIDRemappings& Remappings)
 {
-	for( const FVertexID VertexID : VertexArray.GetElementIDs() )
+	for (const FVertexID VertexID : VertexArray.GetElementIDs())
 	{
-		FMeshVertex& Vertex = VertexArray[ VertexID ];
+		FMeshVertex& Vertex = VertexArray[VertexID];
 
 		// Fix up vertex instance index references in vertices array
-		for( FVertexInstanceID& VertexInstanceID : Vertex.VertexInstanceIDs )
+		for (FVertexInstanceID& VertexInstanceID : Vertex.VertexInstanceIDs)
 		{
-			VertexInstanceID = Remappings.GetRemappedVertexInstanceID( VertexInstanceID );
+			VertexInstanceID = Remappings.GetRemappedVertexInstanceID(VertexInstanceID);
 		}
 
 		// Fix up edge index references in the vertex array
-		for( FEdgeID& EdgeID : Vertex.ConnectedEdgeIDs )
+		for (FEdgeID& EdgeID : Vertex.ConnectedEdgeIDs)
 		{
-			EdgeID = Remappings.GetRemappedEdgeID( EdgeID );
+			EdgeID = Remappings.GetRemappedEdgeID(EdgeID);
 		}
 	}
 
 	// Fix up vertex index references in vertex instance array
-	for( const FVertexInstanceID VertexInstanceID : VertexInstanceArray.GetElementIDs() )
+	for (const FVertexInstanceID VertexInstanceID : VertexInstanceArray.GetElementIDs())
 	{
-		FMeshVertexInstance& VertexInstance = VertexInstanceArray[ VertexInstanceID ];
-
-		VertexInstance.VertexID = Remappings.GetRemappedVertexID( VertexInstance.VertexID );
+		FMeshVertexInstance& VertexInstance = VertexInstanceArray[VertexInstanceID];
 
-		for( FPolygonID& PolygonID : VertexInstance.ConnectedPolygons )
-		{
-			PolygonID = Remappings.GetRemappedPolygonID( PolygonID );
-		}
+		VertexInstance.VertexID = Remappings.GetRemappedVertexID(VertexInstance.VertexID);
 
-		for( FTriangleID& TriangleID : VertexInstance.ConnectedTriangles )
+		for (FTriangleID& TriangleID : VertexInstance.ConnectedTriangles)
 		{
-			TriangleID = Remappings.GetRemappedTriangleID( TriangleID );
+			TriangleID = Remappings.GetRemappedTriangleID(TriangleID);
 		}
 	}
 
-	for( const FEdgeID EdgeID : EdgeArray.GetElementIDs() )
+	for (const FEdgeID EdgeID : EdgeArray.GetElementIDs())
 	{
-		FMeshEdge& Edge = EdgeArray[ EdgeID ];
+		FMeshEdge& Edge = EdgeArray[EdgeID];
 
 		// Fix up vertex index references in Edges array
-		for( int32 Index = 0; Index < 2; Index++ )
-		{
-			Edge.VertexIDs[ Index ] = Remappings.GetRemappedVertexID( Edge.VertexIDs[ Index ] );
-		}
-
-		for( FPolygonID& ConnectedPolygon : Edge.ConnectedPolygons )
+		for (int32 Index = 0; Index < 2; Index++)
 		{
-			ConnectedPolygon = Remappings.GetRemappedPolygonID( ConnectedPolygon );
+			Edge.VertexIDs[Index] = Remappings.GetRemappedVertexID(Edge.VertexIDs[Index]);
 		}
 
-		for( FTriangleID& TriangleID : Edge.ConnectedTriangles )
+		for (FTriangleID& TriangleID : Edge.ConnectedTriangles)
 		{
-			TriangleID = Remappings.GetRemappedTriangleID( TriangleID );
+			TriangleID = Remappings.GetRemappedTriangleID(TriangleID);
 		}
 	}
 
-	for( const FTriangleID TriangleID : TriangleArray.GetElementIDs() )
+	for (const FTriangleID TriangleID : TriangleArray.GetElementIDs())
 	{
-		FMeshTriangle& Triangle = TriangleArray[ TriangleID ];
+		FMeshTriangle& Triangle = TriangleArray[TriangleID];
 		
 		// Fix up vertex instance references in Triangle
-		for( int32 Index = 0; Index < 3; ++Index )
+		for (FVertexInstanceID& VertexInstanceID : Triangle.VertexInstanceIDs)
 		{
-			Triangle.SetVertexInstanceID( Index, Remappings.GetRemappedVertexInstanceID( Triangle.GetVertexInstanceID( Index ) ) );
+			VertexInstanceID = Remappings.GetRemappedVertexInstanceID(VertexInstanceID);
 		}
 
-		Triangle.PolygonID = Remappings.GetRemappedPolygonID( Triangle.PolygonID );
+		Triangle.PolygonID = Remappings.GetRemappedPolygonID(Triangle.PolygonID);
 	}
 
-	for( const FPolygonID PolygonID : PolygonArray.GetElementIDs() )
+	for (const FPolygonID PolygonID : PolygonArray.GetElementIDs())
 	{
-		FMeshPolygon& Polygon = PolygonArray[ PolygonID ];
+		FMeshPolygon& Polygon = PolygonArray[PolygonID];
 
 		// Fix up references to vertex indices in section polygons' contours
-		for( FVertexInstanceID& VertexInstanceID : Polygon.PerimeterContour.VertexInstanceIDs )
+		for (FVertexInstanceID& VertexInstanceID : Polygon.VertexInstanceIDs)
 		{
-			VertexInstanceID = Remappings.GetRemappedVertexInstanceID( VertexInstanceID );
+			VertexInstanceID = Remappings.GetRemappedVertexInstanceID(VertexInstanceID);
 		}
 
-		for( FTriangleID& TriangleID : Polygon.TriangleIDs )
+		for (FTriangleID& TriangleID : Polygon.TriangleIDs)
 		{
-			TriangleID = Remappings.GetRemappedTriangleID( TriangleID );
+			TriangleID = Remappings.GetRemappedTriangleID(TriangleID);
 		}
 
-		Polygon.PolygonGroupID = Remappings.GetRemappedPolygonGroupID( Polygon.PolygonGroupID );
+		Polygon.PolygonGroupID = Remappings.GetRemappedPolygonGroupID(Polygon.PolygonGroupID);
 	}
 
-	for( const FPolygonGroupID PolygonGroupID : PolygonGroupArray.GetElementIDs() )
+	for (const FPolygonGroupID PolygonGroupID : PolygonGroupArray.GetElementIDs())
 	{
-		FMeshPolygonGroup& PolygonGroup = PolygonGroupArray[ PolygonGroupID ];
+		FMeshPolygonGroup& PolygonGroup = PolygonGroupArray[PolygonGroupID];
 
-		for( FPolygonID& Polygon : PolygonGroup.Polygons )
+		for (FPolygonID& Polygon : PolygonGroup.Polygons)
 		{
-			Polygon = Remappings.GetRemappedPolygonID( Polygon );
+			Polygon = Remappings.GetRemappedPolygonID(Polygon);
 		}
 	}
 }
@@ -363,7 +319,7 @@ void FMeshDescription::CreateVertexInstance_Internal(const FVertexInstanceID Ver
 
 void FMeshDescription::DeleteVertexInstance(const FVertexInstanceID VertexInstanceID, TArray<FVertexID>* InOutOrphanedVerticesPtr)
 {
-	check(VertexInstanceArray[VertexInstanceID].ConnectedPolygons.Num() == 0);
+	check(VertexInstanceArray[VertexInstanceID].ConnectedTriangles.Num() == 0);
 	const FVertexID VertexID = VertexInstanceArray[VertexInstanceID].VertexID;
 	verify(VertexArray[VertexID].VertexInstanceIDs.RemoveSingle(VertexInstanceID) == 1);
 	if (InOutOrphanedVerticesPtr && VertexArray[VertexID].VertexInstanceIDs.Num() == 0 && VertexArray[VertexID].ConnectedEdgeIDs.Num() == 0)
@@ -424,8 +380,8 @@ void FMeshDescription::CreateTriangle_Internal(const FTriangleID TriangleID, con
 	FMeshPolygon& Polygon = PolygonArray[PolygonID];
 	PolygonAttributesSet.Insert(PolygonID);
 
-	Polygon.PerimeterContour.VertexInstanceIDs.Reserve(3);
-	Algo::Copy(VertexInstanceIDs, Polygon.PerimeterContour.VertexInstanceIDs);
+	Polygon.VertexInstanceIDs.Reserve(3);
+	Algo::Copy(VertexInstanceIDs, Polygon.VertexInstanceIDs);
 	Polygon.PolygonGroupID = PolygonGroupID;
 	PolygonGroupArray[PolygonGroupID].Polygons.Add(PolygonID);
 
@@ -458,12 +414,6 @@ void FMeshDescription::CreateTriangle_Internal(const FTriangleID TriangleID, con
 
 		check(!EdgeArray[EdgeID].ConnectedTriangles.Contains(TriangleID));
 		EdgeArray[EdgeID].ConnectedTriangles.Add(TriangleID);
-
-		// For backward compatibility only; in 4.24, this will be removed.
-		// Note: internal edges will have the same connected polygon twice.
-		VertexInstanceArray[VertexInstanceID].ConnectedPolygons.Add(PolygonID);
-		EdgeArray[EdgeID].ConnectedPolygons.Add(PolygonID);
-		// End backward compatibility
 	}
 }
 
@@ -493,12 +443,6 @@ void FMeshDescription::DeleteTriangle(const FTriangleID TriangleID, TArray<FEdge
 			verify(VertexInstanceArray[VertexInstanceID].ConnectedTriangles.RemoveSingle(TriangleID) == 1);
 			verify(EdgeArray[EdgeID].ConnectedTriangles.RemoveSingle(TriangleID) == 1);
 
-			// For backward compatibility only; in 4.24, this will be removed.
-			// Note: internal edges will have the same connected polygon twice.
-			verify(VertexInstanceArray[VertexInstanceID].ConnectedPolygons.RemoveSingle(PolygonID) == 1);
-			verify(EdgeArray[EdgeID].ConnectedPolygons.RemoveSingle(PolygonID) == 1);
-			// End backward compatibility
-
 			if (InOutOrphanedVertexInstancesPtr && VertexInstanceArray[VertexInstanceID].ConnectedTriangles.Num() == 0)
 			{
 				InOutOrphanedVertexInstancesPtr->AddUnique(VertexInstanceID);
@@ -519,8 +463,8 @@ void FMeshDescription::DeleteTriangle(const FTriangleID TriangleID, TArray<FEdge
 			InOutOrphanedPolygonGroupsPtr->AddUnique(PolygonGroupID);
 		}
 
-		PolygonArray.Remove( PolygonID );
-		PolygonAttributesSet.Remove( PolygonID );
+		PolygonArray.Remove(PolygonID);
+		PolygonAttributesSet.Remove(PolygonID);
 	}
 	else
 	{
@@ -544,7 +488,7 @@ void FMeshDescription::CreatePolygon_Internal(const FPolygonID PolygonID, const
 
 	FMeshPolygon& Polygon = PolygonArray[PolygonID];
 	const int32 NumVertices = VertexInstanceIDs.Num();
-	Polygon.PerimeterContour.VertexInstanceIDs.SetNumUninitialized(NumVertices);
+	Polygon.VertexInstanceIDs.SetNumUninitialized(NumVertices);
 
 	for (int32 Index = 0; Index < NumVertices; ++Index)
 	{
@@ -553,7 +497,7 @@ void FMeshDescription::CreatePolygon_Internal(const FPolygonID PolygonID, const
 		const FVertexID ThisVertexID = GetVertexInstanceVertex(ThisVertexInstanceID);
 		const FVertexID NextVertexID = GetVertexInstanceVertex(NextVertexInstanceID);
 
-		Polygon.PerimeterContour.VertexInstanceIDs[Index] = ThisVertexInstanceID;
+		Polygon.VertexInstanceIDs[Index] = ThisVertexInstanceID;
 
 		FEdgeID EdgeID = GetVertexPairEdge(ThisVertexID, NextVertexID);
 		if (EdgeID == FEdgeID::Invalid)
@@ -564,13 +508,6 @@ void FMeshDescription::CreatePolygon_Internal(const FPolygonID PolygonID, const
 				OutEdgeIDs->Add(EdgeID);
 			}
 		}
-
-		// For backward compatibility only; in 4.24, this will be removed.
-		check(!VertexInstanceArray[ThisVertexInstanceID].ConnectedPolygons.Contains(PolygonID));
-		VertexInstanceArray[ThisVertexInstanceID].ConnectedPolygons.Add(PolygonID);
-		check(!EdgeArray[EdgeID].ConnectedPolygons.Contains(PolygonID));
-		EdgeArray[EdgeID].ConnectedPolygons.Add(PolygonID);
-		// End backward compatibility
 	}
 
 	check(PolygonGroupID != FPolygonGroupID::Invalid);
@@ -619,10 +556,6 @@ void FMeshDescription::DeletePolygon(const FPolygonID PolygonID, TArray<FEdgeID>
 				{
 					verify(EdgeArray[EdgeID].ConnectedTriangles.RemoveSingle(TriangleID) == 1);
 
-					// For backward compatibility only; in 4.24, this will be removed.
-					verify(EdgeArray[EdgeID].ConnectedPolygons.RemoveSingle(PolygonID) == 1);
-					// End backward compatibility
-
 					if (InOutOrphanedEdgesPtr && EdgeArray[EdgeID].ConnectedTriangles.Num() == 0)
 					{
 						InOutOrphanedEdgesPtr->AddUnique(EdgeID);
@@ -642,31 +575,24 @@ void FMeshDescription::DeletePolygon(const FPolygonID PolygonID, TArray<FEdgeID>
 		TriangleAttributesSet.Remove(TriangleID);
 	}
 
-	// For backward compatibility only; in 4.24, this will be removed.
-	for (const FVertexInstanceID VertexInstanceID : Polygon.PerimeterContour.VertexInstanceIDs)
-	{
-		verify(VertexInstanceArray[VertexInstanceID].ConnectedPolygons.RemoveSingle(PolygonID) == 1);
-	}
-	// End backward compatibility
-
-	FMeshPolygonGroup& PolygonGroup = PolygonGroupArray[ Polygon.PolygonGroupID ];
-	verify( PolygonGroup.Polygons.RemoveSingle( PolygonID ) == 1 );
+	FMeshPolygonGroup& PolygonGroup = PolygonGroupArray[Polygon.PolygonGroupID];
+	verify(PolygonGroup.Polygons.RemoveSingle(PolygonID) == 1);
 
-	if( InOutOrphanedPolygonGroupsPtr && PolygonGroup.Polygons.Num() == 0 )
+	if (InOutOrphanedPolygonGroupsPtr && PolygonGroup.Polygons.Num() == 0)
 	{
-		InOutOrphanedPolygonGroupsPtr->AddUnique( Polygon.PolygonGroupID );
+		InOutOrphanedPolygonGroupsPtr->AddUnique(Polygon.PolygonGroupID);
 	}
 
-	PolygonArray.Remove( PolygonID );
-	PolygonAttributesSet.Remove( PolygonID );
+	PolygonArray.Remove(PolygonID);
+	PolygonAttributesSet.Remove(PolygonID);
 }
 
 
-bool FMeshDescription::IsVertexOrphaned( const FVertexID VertexID ) const
+bool FMeshDescription::IsVertexOrphaned(const FVertexID VertexID) const
 {
-	for( const FVertexInstanceID VertexInstanceID : VertexArray[ VertexID ].VertexInstanceIDs )
+	for (const FVertexInstanceID VertexInstanceID : VertexArray[VertexID].VertexInstanceIDs)
 	{
-		if( VertexInstanceArray[ VertexInstanceID ].ConnectedPolygons.Num() > 0 )
+		if (VertexInstanceArray[VertexInstanceID].ConnectedTriangles.Num() > 0)
 		{
 			return false;
 		}
@@ -710,32 +636,15 @@ FEdgeID FMeshDescription::GetVertexInstancePairEdge(const FVertexInstanceID Vert
 }
 
 
-TArray<FMeshTriangle> FMeshDescription::GetPolygonTriangles(const FPolygonID PolygonID) const
-{
-	TArray<FMeshTriangle> Triangles;
-	for (const FTriangleID TriangleID : PolygonArray[PolygonID].TriangleIDs)
-	{
-		Triangles.Add(TriangleArray[TriangleID]);
-	}
-	return Triangles;
-}
-
-
 void FMeshDescription::SetPolygonVertexInstance(const FPolygonID PolygonID, const int32 PerimeterIndex, const FVertexInstanceID VertexInstanceID)
 {
 	FMeshPolygon& Polygon = PolygonArray[PolygonID];
-	check(PerimeterIndex >= 0 && PerimeterIndex < Polygon.PerimeterContour.VertexInstanceIDs.Num());
+	check(PerimeterIndex >= 0 && PerimeterIndex < Polygon.VertexInstanceIDs.Num());
 
 	// Disconnect old vertex instance from polygon, and connect new one
-	const FVertexInstanceID OldVertexInstanceID = Polygon.PerimeterContour.VertexInstanceIDs[PerimeterIndex];
+	const FVertexInstanceID OldVertexInstanceID = Polygon.VertexInstanceIDs[PerimeterIndex];
 
-	// For backward compatibility only; in 4.24, this will be removed.
-	verify(VertexInstanceArray[OldVertexInstanceID].ConnectedPolygons.RemoveSingle(PolygonID) == 1);
-	check(!VertexInstanceArray[VertexInstanceID].ConnectedPolygons.Contains(PolygonID));
-	VertexInstanceArray[VertexInstanceID].ConnectedPolygons.Add(PolygonID);
-	// End backward compatibility
-
-	Polygon.PerimeterContour.VertexInstanceIDs[PerimeterIndex] = VertexInstanceID;
+	Polygon.VertexInstanceIDs[PerimeterIndex] = VertexInstanceID;
 
 	// Fix up triangle list
 	for (const FTriangleID TriangleID : Polygon.TriangleIDs)
@@ -845,7 +754,7 @@ void FMeshDescription::ComputePolygonTriangulation(const FPolygonID PolygonID)
 	//       page 497 of the book "Real-time Collision Detection", published in 2005.
 
 	FMeshPolygon& Polygon = PolygonArray[PolygonID];
-	const TArray<FVertexInstanceID>& PolygonVertexInstanceIDs = Polygon.PerimeterContour.VertexInstanceIDs;
+	const TArray<FVertexInstanceID>& PolygonVertexInstanceIDs = Polygon.VertexInstanceIDs;
 
 	// Polygon must have at least three vertices/edges
 	const int32 PolygonVertexCount = PolygonVertexInstanceIDs.Num();
@@ -1032,12 +941,6 @@ void FMeshDescription::ComputePolygonTriangulation(const FPolygonID PolygonID)
 					{
 						// This must be an internal edge (as perimeter edges will already be defined)
 						EdgeID = CreateEdge(ThisVertexID, NextVertexID);
-
-						// Internal edges have the same polygon connected either side.
-						// For backward compatibility only; in 4.24, this will be removed.
-						EdgeArray[EdgeID].ConnectedPolygons.Add(PolygonID);
-						EdgeArray[EdgeID].ConnectedPolygons.Add(PolygonID);
-						// End backward compatibility
 					}
 
 					check(!VertexInstanceArray[ThisVertexInstanceID].ConnectedTriangles.Contains(TriangleID));
@@ -1077,167 +980,6 @@ void FMeshDescription::ComputePolygonTriangulation(const FPolygonID PolygonID)
 }
 
 
-void FMeshDescription::ComputePolygonTriangulation(const FPolygonID PolygonID, TArray<FMeshTriangle>& OutTriangles)
-{
-	// NOTE: This polygon triangulation code is partially based on the ear cutting algorithm described on
-	//       page 497 of the book "Real-time Collision Detection", published in 2005.
-
-	struct Local
-	{
-		// Returns true if the triangle formed by the specified three positions has a normal that is facing the opposite direction of the reference normal
-		static inline bool IsTriangleFlipped(const FVector ReferenceNormal, const FVector VertexPositionA, const FVector VertexPositionB, const FVector VertexPositionC)
-		{
-			const FVector TriangleNormal = FVector::CrossProduct(
-				VertexPositionC - VertexPositionA,
-				VertexPositionB - VertexPositionA).GetSafeNormal();
-			return (FVector::DotProduct(ReferenceNormal, TriangleNormal) <= 0.0f);
-		}
-
-	};
-
-
-	// @todo mesheditor: Perhaps should always attempt to triangulate by splitting polygons along the shortest edge, for better determinism.
-
-	//	const FMeshPolygon& Polygon = GetPolygon( PolygonID );
-	const TArray<FVertexInstanceID>& PolygonVertexInstanceIDs = GetPolygonVertexInstances(PolygonID);
-
-	// Polygon must have at least three vertices/edges
-	const int32 PolygonVertexCount = PolygonVertexInstanceIDs.Num();
-	check(PolygonVertexCount >= 3);
-
-	OutTriangles.Reset(PolygonVertexCount-2);
-
-	// If perimeter has 3 vertices, just copy content of perimeter out
-	if (PolygonVertexCount == 3)
-	{
-		OutTriangles.Emplace();
-		FMeshTriangle& Triangle = OutTriangles.Last();
-
-		Triangle.SetVertexInstanceID(0, PolygonVertexInstanceIDs[0]);
-		Triangle.SetVertexInstanceID(1, PolygonVertexInstanceIDs[1]);
-		Triangle.SetVertexInstanceID(2, PolygonVertexInstanceIDs[2]);
-
-		return;
-	}
-
-	// First figure out the polygon normal.  We need this to determine which triangles are convex, so that
-	// we can figure out which ears to clip
-	const FVector PolygonNormal = ComputePolygonNormal(PolygonID);
-
-	// Make a simple linked list array of the previous and next vertex numbers, for each vertex number
-	// in the polygon.  This will just save us having to iterate later on.
-	TArray<int32> PrevVertexNumbers, NextVertexNumbers;
-	TArray<FVector> VertexPositions;
-
-	{
-		const TVertexAttributesRef<FVector> MeshVertexPositions = VertexAttributes().GetAttributesRef<FVector>(MeshAttribute::Vertex::Position);
-		PrevVertexNumbers.SetNumUninitialized(PolygonVertexCount, false);
-		NextVertexNumbers.SetNumUninitialized(PolygonVertexCount, false);
-		VertexPositions.SetNumUninitialized(PolygonVertexCount, false);
-
-		for (int32 VertexNumber = 0; VertexNumber < PolygonVertexCount; ++VertexNumber)
-		{
-			PrevVertexNumbers[VertexNumber] = VertexNumber - 1;
-			NextVertexNumbers[VertexNumber] = VertexNumber + 1;
-
-			VertexPositions[VertexNumber] = MeshVertexPositions[GetVertexInstanceVertex(PolygonVertexInstanceIDs[VertexNumber])];
-		}
-		PrevVertexNumbers[0] = PolygonVertexCount - 1;
-		NextVertexNumbers[PolygonVertexCount - 1] = 0;
-	}
-
-	int32 EarVertexNumber = 0;
-	int32 EarTestCount = 0;
-	for (int32 RemainingVertexCount = PolygonVertexCount; RemainingVertexCount >= 3; )
-	{
-		bool bIsEar = true;
-
-		// If we're down to only a triangle, just treat it as an ear.  Also, if we've tried every possible candidate
-		// vertex looking for an ear, go ahead and just treat the current vertex as an ear.  This can happen when
-		// vertices are colinear or other degenerate cases.
-		if (RemainingVertexCount > 3 && EarTestCount < RemainingVertexCount)
-		{
-			const FVector PrevVertexPosition = VertexPositions[PrevVertexNumbers[EarVertexNumber]];
-			const FVector EarVertexPosition = VertexPositions[EarVertexNumber];
-			const FVector NextVertexPosition = VertexPositions[NextVertexNumbers[EarVertexNumber]];
-
-			// Figure out whether the potential ear triangle is facing the same direction as the polygon
-			// itself.  If it's facing the opposite direction, then we're dealing with a concave triangle
-			// and we'll skip it for now.
-			if (!Local::IsTriangleFlipped(
-				PolygonNormal,
-				PrevVertexPosition,
-				EarVertexPosition,
-				NextVertexPosition))
-			{
-				int32 TestVertexNumber = NextVertexNumbers[NextVertexNumbers[EarVertexNumber]];
-
-				do
-				{
-					// Test every other remaining vertex to make sure that it doesn't lie inside our potential ear
-					// triangle.  If we find a vertex that's inside the triangle, then it cannot actually be an ear.
-					const FVector TestVertexPosition = VertexPositions[TestVertexNumber];
-					if (PointInTriangle(
-						PrevVertexPosition,
-						EarVertexPosition,
-						NextVertexPosition,
-						TestVertexPosition,
-						SMALL_NUMBER))
-					{
-						bIsEar = false;
-						break;
-					}
-
-					TestVertexNumber = NextVertexNumbers[TestVertexNumber];
-				} while (TestVertexNumber != PrevVertexNumbers[EarVertexNumber]);
-			}
-			else
-			{
-				bIsEar = false;
-			}
-		}
-
-		if (bIsEar)
-		{
-			// OK, we found an ear!  Let's save this triangle in our output buffer.
-			{
-				OutTriangles.Emplace();
-				FMeshTriangle& Triangle = OutTriangles.Last();
-
-				Triangle.SetVertexInstanceID(0, PolygonVertexInstanceIDs[PrevVertexNumbers[EarVertexNumber]]);
-				Triangle.SetVertexInstanceID(1, PolygonVertexInstanceIDs[EarVertexNumber]);
-				Triangle.SetVertexInstanceID(2, PolygonVertexInstanceIDs[NextVertexNumbers[EarVertexNumber]]);
-			}
-
-			// Update our linked list.  We're effectively cutting off the ear by pointing the ear vertex's neighbors to
-			// point at their next sequential neighbor, and reducing the remaining vertex count by one.
-			{
-				NextVertexNumbers[PrevVertexNumbers[EarVertexNumber]] = NextVertexNumbers[EarVertexNumber];
-				PrevVertexNumbers[NextVertexNumbers[EarVertexNumber]] = PrevVertexNumbers[EarVertexNumber];
-				--RemainingVertexCount;
-			}
-
-			// Move on to the previous vertex in the list, now that this vertex was cut
-			EarVertexNumber = PrevVertexNumbers[EarVertexNumber];
-
-			EarTestCount = 0;
-		}
-		else
-		{
-			// The vertex is not the ear vertex, because it formed a triangle that either had a normal which pointed in the opposite direction
-			// of the polygon, or at least one of the other polygon vertices was found to be inside the triangle.  Move on to the next vertex.
-			EarVertexNumber = NextVertexNumbers[EarVertexNumber];
-
-			// Keep track of how many ear vertices we've tested, so that if we exhaust all remaining vertices, we can
-			// fall back to clipping the triangle and adding it to our mesh anyway.  This is important for degenerate cases.
-			++EarTestCount;
-		}
-	}
-
-	check(OutTriangles.Num() > 0);
-}
-
-
 FBoxSphereBounds FMeshDescription::GetBounds() const
 {
 	FBoxSphereBounds BoundingBoxAndSphere;
@@ -1321,232 +1063,21 @@ namespace MeshAttribute_
 }
 
 
-
-bool FMeshDescription::ComputePolygonTangentsAndNormals(const FPolygonID PolygonID
-	, float ComparisonThreshold
-	, const TVertexAttributesRef<FVector> VertexPositions
-	, const TVertexInstanceAttributesRef<FVector2D> VertexUVs
-	, TPolygonAttributesRef<FVector> PolygonNormals
-	, TPolygonAttributesRef<FVector> PolygonTangents
-	, TPolygonAttributesRef<FVector> PolygonBinormals
-	, TPolygonAttributesRef<FVector> PolygonCenters)
-{
-	bool bValidNTBs = true;
-	// Calculate the center of this polygon
-	FVector Center = FVector::ZeroVector;
-	const TArray<FVertexInstanceID>& VertexInstanceIDs = GetPolygonVertexInstances(PolygonID);
-	for (const FVertexInstanceID VertexInstanceID : VertexInstanceIDs)
-	{
-		Center += VertexPositions[GetVertexInstanceVertex(VertexInstanceID)];
-	}
-	Center /= float(VertexInstanceIDs.Num());
-
-	// Calculate the tangent basis for the polygon, based on the average of all constituent triangles
-	FVector Normal = FVector::ZeroVector;
-	FVector Tangent = FVector::ZeroVector;
-	FVector Binormal = FVector::ZeroVector;
-
-	for (const FTriangleID TriangleID : GetPolygonTriangleIDs(PolygonID))
-	{
-		const FMeshTriangle& Triangle = TriangleArray[ TriangleID ];
-		const FVertexID VertexID0 = GetVertexInstanceVertex(Triangle.GetVertexInstanceID(0));
-		const FVertexID VertexID1 = GetVertexInstanceVertex(Triangle.GetVertexInstanceID(1));
-		const FVertexID VertexID2 = GetVertexInstanceVertex(Triangle.GetVertexInstanceID(2));
-
-		const FVector DPosition1 = VertexPositions[VertexID1] - VertexPositions[VertexID0];
-		const FVector DPosition2 = VertexPositions[VertexID2] - VertexPositions[VertexID0];
-
-		const FVector2D DUV1 = VertexUVs[Triangle.GetVertexInstanceID(1)] - VertexUVs[Triangle.GetVertexInstanceID(0)];
-		const FVector2D DUV2 = VertexUVs[Triangle.GetVertexInstanceID(2)] - VertexUVs[Triangle.GetVertexInstanceID(0)];
-
-		// We have a left-handed coordinate system, but a counter-clockwise winding order
-		// Hence normal calculation has to take the triangle vectors cross product in reverse.
-		FVector TmpNormal = FVector::CrossProduct(DPosition2, DPosition1);
-		if (!TmpNormal.IsNearlyZero(ComparisonThreshold) && !TmpNormal.ContainsNaN())
-		{
-			Normal += TmpNormal;
-			// ...and tangent space seems to be right-handed.
-			const float DetUV = FVector2D::CrossProduct(DUV1, DUV2);
-			const float InvDetUV = (DetUV == 0.0f) ? 0.0f : 1.0f / DetUV;
-
-			Tangent += (DPosition1 * DUV2.Y - DPosition2 * DUV1.Y) * InvDetUV;
-			Binormal += (DPosition2 * DUV1.X - DPosition1 * DUV2.X) * InvDetUV;
-		}
-		else
-		{
-			//The polygon is degenerated
-			bValidNTBs = false;
-		}
-	}
-
-	PolygonNormals[PolygonID] = Normal.GetSafeNormal();
-	PolygonTangents[PolygonID] = Tangent.GetSafeNormal();
-	PolygonBinormals[PolygonID] = Binormal.GetSafeNormal();
-	PolygonCenters[PolygonID] = Center;
-
-	return bValidNTBs;
-}
-
-
-void FMeshDescription::ComputePolygonTangentsAndNormals(const TArray<FPolygonID>& PolygonIDs, float ComparisonThreshold)
-{
-	const TVertexAttributesRef<FVector> VertexPositions = VertexAttributes().GetAttributesRef<FVector>(MeshAttribute::Vertex::Position);
-	const TVertexInstanceAttributesRef<FVector2D> VertexUVs = VertexInstanceAttributes().GetAttributesRef<FVector2D>(MeshAttribute_::VertexInstance::TextureCoordinate);
-	TPolygonAttributesRef<FVector> PolygonNormals = PolygonAttributes().GetAttributesRef<FVector>(MeshAttribute_::Polygon::Normal);
-	TPolygonAttributesRef<FVector> PolygonTangents = PolygonAttributes().GetAttributesRef<FVector>(MeshAttribute_::Polygon::Tangent);
-	TPolygonAttributesRef<FVector> PolygonBinormals = PolygonAttributes().GetAttributesRef<FVector>(MeshAttribute_::Polygon::Binormal);
-	TPolygonAttributesRef<FVector> PolygonCenters = PolygonAttributes().GetAttributesRef<FVector>(MeshAttribute_::Polygon::Center);
-
-	TArray<FPolygonID> DegeneratePolygonIDs;
-	for (const FPolygonID PolygonID : PolygonIDs)
-	{
-		if (!ComputePolygonTangentsAndNormals(PolygonID, ComparisonThreshold, VertexPositions, VertexUVs, PolygonNormals, PolygonTangents, PolygonBinormals, PolygonCenters))
-		{
-			DegeneratePolygonIDs.Add(PolygonID);
-		}
-	}
-	//Remove degenerated polygons
-	//Delete the degenerated polygons. The array is fill only if the remove degenerated option is turn on.
-	if (DegeneratePolygonIDs.Num() > 0)
-	{
-		TArray<FEdgeID> OrphanedEdges;
-		TArray<FVertexInstanceID> OrphanedVertexInstances;
-		TArray<FPolygonGroupID> OrphanedPolygonGroups;
-		TArray<FVertexID> OrphanedVertices;
-		for (FPolygonID PolygonID : DegeneratePolygonIDs)
-		{
-			DeletePolygon(PolygonID, &OrphanedEdges, &OrphanedVertexInstances, &OrphanedPolygonGroups);
-		}
-		for (FPolygonGroupID PolygonGroupID : OrphanedPolygonGroups)
-		{
-			DeletePolygonGroup(PolygonGroupID);
-		}
-		for (FVertexInstanceID VertexInstanceID : OrphanedVertexInstances)
-		{
-			DeleteVertexInstance(VertexInstanceID, &OrphanedVertices);
-		}
-		for (FEdgeID EdgeID : OrphanedEdges)
-		{
-			DeleteEdge(EdgeID, &OrphanedVertices);
-		}
-		for (FVertexID VertexID : OrphanedVertices)
-		{
-			DeleteVertex(VertexID);
-		}
-		//Compact and Remap IDs so we have clean ID from 0 to n since we just erase some polygons
-		//The render build need to have compact ID
-		FElementIDRemappings RemappingInfos;
-		Compact(RemappingInfos);
-	}
-}
-
-
-void FMeshDescription::ComputePolygonTangentsAndNormals(float ComparisonThreshold)
-{
-	TArray<FPolygonID> PolygonsToComputeNTBs;
-	PolygonsToComputeNTBs.Reserve(Polygons().Num());
-	for (const FPolygonID& PolygonID : Polygons().GetElementIDs())
-	{
-		PolygonsToComputeNTBs.Add(PolygonID);
-	}
-	ComputePolygonTangentsAndNormals(PolygonsToComputeNTBs, ComparisonThreshold);
-}
-
-
-void FMeshDescription::GetConnectedSoftEdges(const FVertexID VertexID, TArray<FEdgeID>& OutConnectedSoftEdges) const
-{
-	OutConnectedSoftEdges.Reset();
-
-	TEdgeAttributesConstRef<bool> EdgeHardnesses = EdgeAttributes().GetAttributesRef<bool>(MeshAttribute_::Edge::IsHard);
-	for (const FEdgeID ConnectedEdgeID : VertexArray[VertexID].ConnectedEdgeIDs)
-	{
-		if (!EdgeHardnesses[ConnectedEdgeID])
-		{
-			OutConnectedSoftEdges.Add(ConnectedEdgeID);
-		}
-	}
-}
-
-
-void FMeshDescription::GetPolygonsInSameSoftEdgedGroupAsPolygon(const FPolygonID PolygonID, const TArray<FPolygonID>& CandidatePolygonIDs, const TArray<FEdgeID>& SoftEdgeIDs, TArray<FPolygonID>& OutPolygonIDs) const
-{
-	// The aim of this method is:
-	// - given a polygon ID,
-	// - given a set of candidate polygons connected to the same vertex (which should include the polygon ID),
-	// - given a set of soft edges connected to the same vertex,
-	// return the polygon IDs which form an adjacent run without crossing a hard edge.
-
-	OutPolygonIDs.Reset();
-
-	// Maintain a list of polygon IDs to be examined. Adjacents are added to the list if suitable.
-	// Add the start poly here.
-	TArray<FPolygonID> PolygonsToCheck;
-	PolygonsToCheck.Reset(CandidatePolygonIDs.Num());
-	PolygonsToCheck.Add(PolygonID);
-
-	int32 Index = 0;
-	while (Index < PolygonsToCheck.Num())
-	{
-		const FPolygonID PolygonToCheck = PolygonsToCheck[Index];
-		Index++;
-
-		if (CandidatePolygonIDs.Contains(PolygonToCheck))
-		{
-			OutPolygonIDs.Add(PolygonToCheck);
-
-			// Now look at its adjacent polygons. If they are joined by a soft edge which includes the vertex we're interested in, we want to consider them.
-			// We take a shortcut by doing this process in reverse: we already know all the soft edges we are interested in, so check if any of them
-			// have the current polygon as an adjacent.
-			for (const FEdgeID SoftEdgeID : SoftEdgeIDs)
-			{
-				const TArray<FPolygonID>& EdgeConnectedPolygons = GetEdgeConnectedPolygons(SoftEdgeID);
-				if (EdgeConnectedPolygons.Contains(PolygonToCheck))
-				{
-					for (const FPolygonID AdjacentPolygon : EdgeConnectedPolygons)
-					{
-						// Only add new polygons which haven't yet been added to the list. This prevents circular runs of polygons triggering infinite loops.
-						PolygonsToCheck.AddUnique(AdjacentPolygon);
-					}
-				}
-			}
-		}
-	}
-}
-
-
-void FMeshDescription::GetVertexConnectedPolygonsInSameSoftEdgedGroup(const FVertexID VertexID, const FPolygonID PolygonID, TArray<FPolygonID>& OutPolygonIDs) const
-{
-	// The aim here is to determine which polygons form part of the same soft edged group as the polygons attached to this vertex.
-	// They should all contribute to the final vertex instance normal.
-
-	// Get all polygons connected to this vertex.
-	TArray<FPolygonID> ConnectedPolygons;
-	GetVertexConnectedPolygons(VertexID, ConnectedPolygons);
-
-	// Cache a list of all soft edges which share this vertex.
-	// We're only interested in finding adjacent polygons which are not the other side of a hard edge.
-	TArray<FEdgeID> ConnectedSoftEdges;
-	GetConnectedSoftEdges(VertexID, ConnectedSoftEdges);
-
-	GetPolygonsInSameSoftEdgedGroupAsPolygon(PolygonID, ConnectedPolygons, ConnectedSoftEdges, OutPolygonIDs);
-}
-
-
 float FMeshDescription::GetPolygonCornerAngleForVertex(const FPolygonID PolygonID, const FVertexID VertexID) const
 {
 	const FMeshPolygon& Polygon = PolygonArray[PolygonID];
 
 	// Lambda function which returns the inner angle at a given index on a polygon contour
-	auto GetContourAngle = [this](const FMeshPolygonContour& Contour, const int32 ContourIndex)
+	auto GetContourAngle = [this](const TArray<FVertexInstanceID>& VertexInstanceIDs, const int32 ContourIndex)
 	{
-		const int32 NumVertices = Contour.VertexInstanceIDs.Num();
+		const int32 NumVertices = VertexInstanceIDs.Num();
 
 		const int32 PrevIndex = (ContourIndex + NumVertices - 1) % NumVertices;
 		const int32 NextIndex = (ContourIndex + 1) % NumVertices;
 
-		const FVertexID PrevVertexID = GetVertexInstanceVertex(Contour.VertexInstanceIDs[PrevIndex]);
-		const FVertexID ThisVertexID = GetVertexInstanceVertex(Contour.VertexInstanceIDs[ContourIndex]);
-		const FVertexID NextVertexID = GetVertexInstanceVertex(Contour.VertexInstanceIDs[NextIndex]);
+		const FVertexID PrevVertexID = GetVertexInstanceVertex(VertexInstanceIDs[PrevIndex]);
+		const FVertexID ThisVertexID = GetVertexInstanceVertex(VertexInstanceIDs[ContourIndex]);
+		const FVertexID NextVertexID = GetVertexInstanceVertex(VertexInstanceIDs[NextIndex]);
 
 		TVertexAttributesConstRef<FVector> VertexPositions = VertexAttributes().GetAttributesRef<FVector>(MeshAttribute::Vertex::Position);
 
@@ -1567,11 +1098,11 @@ float FMeshDescription::GetPolygonCornerAngleForVertex(const FPolygonID PolygonI
 	};
 
 	// First look for the vertex instance in the perimeter
-	int32 ContourIndex = Polygon.PerimeterContour.VertexInstanceIDs.IndexOfByPredicate(IsVertexInstancedFromThisVertex);
+	int32 ContourIndex = Polygon.VertexInstanceIDs.IndexOfByPredicate(IsVertexInstancedFromThisVertex);
 	if (ContourIndex != INDEX_NONE)
 	{
 		// Return the internal angle if found
-		return GetContourAngle(Polygon.PerimeterContour, ContourIndex);
+		return GetContourAngle(Polygon.VertexInstanceIDs, ContourIndex);
 	}
 
 	// Found nothing; return 0
@@ -1592,372 +1123,21 @@ FBox FMeshDescription::ComputeBoundingBox() const
 	return BoundingBox;
 }
 
-void FMeshDescription::ComputeTangentsAndNormals(const FVertexInstanceID VertexInstanceID
-	, EComputeNTBsOptions ComputeNTBsOptions
-	, const TPolygonAttributesRef<FVector> PolygonNormals
-	, const TPolygonAttributesRef<FVector> PolygonTangents
-	, const TPolygonAttributesRef<FVector> PolygonBinormals
-	, TVertexInstanceAttributesRef<FVector> VertexNormals
-	, TVertexInstanceAttributesRef<FVector> VertexTangents
-	, TVertexInstanceAttributesRef<float> VertexBinormalSigns)
-{
-	bool bComputeNormals = !!(ComputeNTBsOptions & EComputeNTBsOptions::Normals);
-	bool bComputeTangents = !!(ComputeNTBsOptions & EComputeNTBsOptions::Tangents);
-	bool bUseWeightedNormals = !!(ComputeNTBsOptions & EComputeNTBsOptions::WeightedNTBs);
-
-	FVector Normal = FVector::ZeroVector;
-	FVector Tangent = FVector::ZeroVector;
-	FVector Binormal = FVector::ZeroVector;
-
-	FVector& NormalRef = VertexNormals[VertexInstanceID];
-	FVector& TangentRef = VertexTangents[VertexInstanceID];
-	float& BinormalRef = VertexBinormalSigns[VertexInstanceID];
-
-	if (!bComputeNormals && !bComputeTangents)
-	{
-		//Nothing to compute
-		return;
-	}
-
-	const FVertexID VertexID = GetVertexInstanceVertex(VertexInstanceID);
-
-	if (bComputeNormals || NormalRef.IsNearlyZero())
-	{
-		// Get all polygons connected to this vertex instance
-		TArray<FPolygonID> AllConnectedPolygons;
-		const TArray<FPolygonID>& VertexInstanceConnectedPolygons = GetVertexInstanceConnectedPolygons(VertexInstanceID);
-		check(VertexInstanceConnectedPolygons.Num() > 0);
-		// Add also any in the same smoothing group connected to a different vertex instance
-		// (as they still have influence over the normal).
-		GetVertexConnectedPolygonsInSameSoftEdgedGroup(VertexID, VertexInstanceConnectedPolygons[0], AllConnectedPolygons);
-		// The vertex instance normal is computed as a sum of all connected polygons' normals, weighted by the angle they make with the vertex
-		for (const FPolygonID ConnectedPolygonID : AllConnectedPolygons)
-		{
-			const float Angle = bUseWeightedNormals ? GetPolygonCornerAngleForVertex(ConnectedPolygonID, VertexID) : 1.0f;
-
-			Normal += PolygonNormals[ConnectedPolygonID] * Angle;
-
-			// If this polygon is actually connected to the vertex instance we're processing, also include its contributions towards the tangent
-			if (VertexInstanceConnectedPolygons.Contains(ConnectedPolygonID))
-			{
-				Tangent += PolygonTangents[ConnectedPolygonID] * Angle;
-				Binormal += PolygonBinormals[ConnectedPolygonID] * Angle;
-			}
-		}
-		// Normalize Normal
-		Normal = Normal.GetSafeNormal();
-	}
-	else
-	{
-		//We use existing normals so just use all polygons having a vertex instance at the same location sharing the same normals
-		Normal = NormalRef;
-		TArray<FVertexInstanceID> VertexInstanceIDs = GetVertexVertexInstances(VertexID);
-		for (const FVertexInstanceID& ConnectedVertexInstanceID : VertexInstanceIDs)
-		{
-			if (ConnectedVertexInstanceID != VertexInstanceID && !VertexNormals[ConnectedVertexInstanceID].Equals(Normal))
-			{
-				continue;
-			}
-
-			const TArray<FPolygonID>& ConnectedPolygons = GetVertexInstanceConnectedPolygons(ConnectedVertexInstanceID);
-			for (const FPolygonID ConnectedPolygonID : ConnectedPolygons)
-			{
-				const float Angle = bUseWeightedNormals ? GetPolygonCornerAngleForVertex(ConnectedPolygonID, VertexID) : 1.0f;
-
-				// If this polygon is actually connected to the vertex instance we're processing, also include its contributions towards the tangent
-				Tangent += PolygonTangents[ConnectedPolygonID] * Angle;
-				Binormal += PolygonBinormals[ConnectedPolygonID] * Angle;
-			}
-		}
-	}
-
-
-	float BinormalSign = 1.0f;
-	if (bComputeTangents)
-	{
-		// Make Tangent orthonormal to Normal.
-		// This is a quicker method than normalizing Tangent, taking the cross product Normal X Tangent, and then a further cross product with that result
-		Tangent = (Tangent - Normal * FVector::DotProduct(Normal, Tangent)).GetSafeNormal();
-
-		// Calculate binormal sign
-		BinormalSign = (FVector::DotProduct(FVector::CrossProduct(Normal, Tangent), Binormal) < 0.0f) ? -1.0f : 1.0f;
-	}
-
-	//Set the value that need to be set
-	if (NormalRef.IsNearlyZero())
-	{
-		NormalRef = Normal;
-	}
-	if (bComputeTangents)
-	{
-		if (TangentRef.IsNearlyZero())
-		{
-			TangentRef = Tangent;
-		}
-		if (FMath::IsNearlyZero(BinormalRef))
-		{
-			BinormalRef = BinormalSign;
-		}
-	}
-}
-
-
-void FMeshDescription::ComputeTangentsAndNormals(const TArray<FVertexInstanceID>& VertexInstanceIDs, EComputeNTBsOptions ComputeNTBsOptions)
-{
-	const TPolygonAttributesRef<FVector> PolygonNormals = PolygonAttributes().GetAttributesRef<FVector>(MeshAttribute_::Polygon::Normal);
-	const TPolygonAttributesRef<FVector> PolygonTangents = PolygonAttributes().GetAttributesRef<FVector>(MeshAttribute_::Polygon::Tangent);
-	const TPolygonAttributesRef<FVector> PolygonBinormals = PolygonAttributes().GetAttributesRef<FVector>(MeshAttribute_::Polygon::Binormal);
-
-	TVertexInstanceAttributesRef<FVector> VertexNormals = VertexInstanceAttributes().GetAttributesRef<FVector>(MeshAttribute_::VertexInstance::Normal);
-	TVertexInstanceAttributesRef<FVector> VertexTangents = VertexInstanceAttributes().GetAttributesRef<FVector>(MeshAttribute_::VertexInstance::Tangent);
-	TVertexInstanceAttributesRef<float> VertexBinormalSigns = VertexInstanceAttributes().GetAttributesRef<float>(MeshAttribute_::VertexInstance::BinormalSign);
-
-	for (const FVertexInstanceID VertexInstanceID : VertexInstanceIDs)
-	{
-		ComputeTangentsAndNormals(VertexInstanceID, ComputeNTBsOptions, PolygonNormals, PolygonTangents, PolygonBinormals, VertexNormals, VertexTangents, VertexBinormalSigns);
-	}
-}
-
-
-void FMeshDescription::ComputeTangentsAndNormals(EComputeNTBsOptions ComputeNTBsOptions)
-{
-	const TPolygonAttributesRef<FVector> PolygonNormals = PolygonAttributes().GetAttributesRef<FVector>(MeshAttribute_::Polygon::Normal);
-	const TPolygonAttributesRef<FVector> PolygonTangents = PolygonAttributes().GetAttributesRef<FVector>(MeshAttribute_::Polygon::Tangent);
-	const TPolygonAttributesRef<FVector> PolygonBinormals = PolygonAttributes().GetAttributesRef<FVector>(MeshAttribute_::Polygon::Binormal);
-
-	TVertexInstanceAttributesRef<FVector> VertexNormals = VertexInstanceAttributes().GetAttributesRef<FVector>(MeshAttribute_::VertexInstance::Normal);
-	TVertexInstanceAttributesRef<FVector> VertexTangents = VertexInstanceAttributes().GetAttributesRef<FVector>(MeshAttribute_::VertexInstance::Tangent);
-	TVertexInstanceAttributesRef<float> VertexBinormalSigns = VertexInstanceAttributes().GetAttributesRef<float>(MeshAttribute_::VertexInstance::BinormalSign);
-
-	for (const FVertexInstanceID VertexInstanceID : VertexInstances().GetElementIDs())
-	{
-		ComputeTangentsAndNormals(VertexInstanceID, ComputeNTBsOptions, PolygonNormals, PolygonTangents, PolygonBinormals, VertexNormals, VertexTangents, VertexBinormalSigns);
-	}
-}
-
-
-void FMeshDescription::DetermineEdgeHardnessesFromVertexInstanceNormals( const float Tolerance )
-{
-	const TVertexInstanceAttributesRef<FVector> VertexNormals = VertexInstanceAttributes().GetAttributesRef<FVector>( MeshAttribute_::VertexInstance::Normal );
-	TEdgeAttributesRef<bool> EdgeHardnesses = EdgeAttributes().GetAttributesRef<bool>( MeshAttribute_::Edge::IsHard );
-
-	// Holds unique vertex instance IDs for a given edge vertex
-	// @todo: use TMemStackAllocator or similar to avoid expensive allocations
-	TArray<FVertexInstanceID> UniqueVertexInstanceIDs;
-
-	for( const FEdgeID EdgeID : Edges().GetElementIDs() )
-	{
-		// Get list of polygons connected to this edge
-		const TArray<FPolygonID>& ConnectedPolygonIDs = GetEdgeConnectedPolygons( EdgeID );
-		if( ConnectedPolygonIDs.Num() == 0 )
-		{
-			// What does it mean if an edge has no connected polygons? For now we just skip it
-			continue;
-		}
-
-		// Assume by default that the edge is soft - but as soon as any vertex instance belonging to a connected polygon
-		// has a distinct normal from the others (within the given tolerance), we mark it as hard.
-		// The exception is if an edge has exactly one connected polygon: in this case we automatically deem it a hard edge.
-		bool bEdgeIsHard = ( ConnectedPolygonIDs.Num() == 1 );
-
-		// Examine vertices on each end of the edge, if we haven't yet identified it as 'hard'
-		for( int32 VertexIndex = 0; !bEdgeIsHard && VertexIndex < 2; ++VertexIndex )
-		{
-			const FVertexID VertexID = GetEdgeVertex( EdgeID, VertexIndex );
-
-			const int32 ReservedElements = 4;
-			UniqueVertexInstanceIDs.Reset( ReservedElements );
-
-			// Get a list of all vertex instances for this vertex which form part of any polygon connected to the edge
-			for( const FVertexInstanceID VertexInstanceID : GetVertexVertexInstances( VertexID ) )
-			{
-				for( const FPolygonID PolygonID : GetVertexInstanceConnectedPolygons( VertexInstanceID ) )
-				{
-					if( ConnectedPolygonIDs.Contains( PolygonID ) )
-					{
-						UniqueVertexInstanceIDs.AddUnique( VertexInstanceID );
-						break;
-					}
-				}
-			}
-			check( UniqueVertexInstanceIDs.Num() > 0 );
-
-			// First unique vertex instance is used as a reference against which the others are compared.
-			// (not a perfect approach: really the 'median' should be used as a reference)
-			const FVector ReferenceNormal = VertexNormals[ UniqueVertexInstanceIDs[ 0 ] ];
-			for( int32 Index = 1; Index < UniqueVertexInstanceIDs.Num(); ++Index )
-			{
-				if( !VertexNormals[ UniqueVertexInstanceIDs[ Index ] ].Equals( ReferenceNormal, Tolerance ) )
-				{
-					bEdgeIsHard = true;
-					break;
-				}
-			}
-		}
-
-		EdgeHardnesses[ EdgeID ] = bEdgeIsHard;
-	}
-}
-
-
-void FMeshDescription::DetermineUVSeamsFromUVs( const int32 UVIndex, const float Tolerance )
-{
-	const TVertexInstanceAttributesRef<FVector2D> VertexUVs = VertexInstanceAttributes().GetAttributesRef<FVector2D>( MeshAttribute_::VertexInstance::TextureCoordinate );
-	TEdgeAttributesRef<bool> EdgeUVSeams = EdgeAttributes().GetAttributesRef<bool>( MeshAttribute_::Edge::IsUVSeam );
-
-	// Holds unique vertex instance IDs for a given edge vertex
-	// @todo: use TMemStackAllocator or similar to avoid expensive allocations
-	TArray<FVertexInstanceID> UniqueVertexInstanceIDs;
-
-	for( const FEdgeID EdgeID : Edges().GetElementIDs() )
-	{
-		// Get list of polygons connected to this edge
-		const TArray<FPolygonID>& ConnectedPolygonIDs = GetEdgeConnectedPolygons( EdgeID );
-		if( ConnectedPolygonIDs.Num() == 0 )
-		{
-			// What does it mean if an edge has no connected polygons? For now we just skip it
-			continue;
-		}
-
-		// Assume by default that the edge is not a UV seam - but as soon as any vertex instance belonging to a connected polygon
-		// has a distinct UV from the others (within the given tolerance), we mark it as a UV seam.
-		bool bEdgeIsUVSeam = false;
-
-		// Examine vertices on each end of the edge, if we haven't yet identified it as a UV seam
-		for( int32 VertexIndex = 0; !bEdgeIsUVSeam && VertexIndex < 2; ++VertexIndex )
-		{
-			const FVertexID VertexID = GetEdgeVertex( EdgeID, VertexIndex );
-
-			const int32 ReservedElements = 4;
-			UniqueVertexInstanceIDs.Reset( ReservedElements );
-
-			// Get a list of all vertex instances for this vertex which form part of any polygon connected to the edge
-			for( const FVertexInstanceID VertexInstanceID : GetVertexVertexInstances( VertexID ) )
-			{
-				for( const FPolygonID PolygonID : GetVertexInstanceConnectedPolygons( VertexInstanceID ) )
-				{
-					if( ConnectedPolygonIDs.Contains( PolygonID ) )
-					{
-						UniqueVertexInstanceIDs.AddUnique( VertexInstanceID );
-						break;
-					}
-				}
-			}
-			check( UniqueVertexInstanceIDs.Num() > 0 );
-
-			// First unique vertex instance is used as a reference against which the others are compared.
-			// (not a perfect approach: really the 'median' should be used as a reference)
-			const FVector2D ReferenceUV = VertexUVs.Get( UniqueVertexInstanceIDs[ 0 ], UVIndex );
-			for( int32 Index = 1; Index < UniqueVertexInstanceIDs.Num(); ++Index )
-			{
-				if( !VertexUVs.Get( UniqueVertexInstanceIDs[ Index ], UVIndex ).Equals( ReferenceUV, Tolerance ) )
-				{
-					bEdgeIsUVSeam = true;
-					break;
-				}
-			}
-		}
-
-		EdgeUVSeams[ EdgeID ] = bEdgeIsUVSeam;
-	}
-}
-
-
-void FMeshDescription::GetPolygonsInSameChartAsPolygon( const FPolygonID PolygonID, TArray<FPolygonID>& OutPolygonIDs )
-{
-	const TEdgeAttributesRef<bool> EdgeUVSeams = EdgeAttributes().GetAttributesRef<bool>( MeshAttribute_::Edge::IsUVSeam );
-	const int32 NumPolygons = Polygons().Num();
-
-	// This holds the results - all polygon IDs which are in the same UV chart
-	OutPolygonIDs.Reset( NumPolygons );
-
-	// This holds all the polygons we need to check, and those we have already checked so we don't add duplicates
-	// @todo: use TMemStackAllocator or similar to avoid expensive allocations
-	TArray<FPolygonID> PolygonsToCheck;
-	PolygonsToCheck.Reserve( NumPolygons );
-
-	// Add the initial polygon
-	PolygonsToCheck.Add( PolygonID );
-
-	int32 Index = 0;
-	while( Index < PolygonsToCheck.Num() )
-	{
-		// Process the next polygon to be checked. If it's in this list, we already know it's one of the results. Now we have to check the neighbors.
-		const FPolygonID PolygonToCheck = PolygonsToCheck[ Index ];
-		OutPolygonIDs.Add( PolygonToCheck );
-		Index++;
-
-		// Iterate through edges of the polygon
-		const TArray<FVertexInstanceID>& VertexInstanceIDs = GetPolygonVertexInstances( PolygonToCheck );
-		FVertexID LastVertexID = GetVertexInstanceVertex( VertexInstanceIDs.Last() );
-		for( const FVertexInstanceID VertexInstanceID : VertexInstanceIDs )
-		{
-			const FVertexID VertexID = GetVertexInstanceVertex( VertexInstanceID );
-			const FEdgeID EdgeID = GetVertexPairEdge( VertexID, LastVertexID );
-			if( EdgeID != FEdgeID::Invalid && !EdgeUVSeams[ EdgeID ] )
-			{
-				// If it's a valid edge and not a UV seam, check its connected polygons
-				const TArray<FPolygonID>& ConnectedPolygonIDs = GetEdgeConnectedPolygons( EdgeID );
-				for( const FPolygonID ConnectedPolygonID : ConnectedPolygonIDs )
-				{
-					// Add polygons which aren't the one being checked, and haven't already been added to the list
-					if( ConnectedPolygonID != PolygonToCheck && !PolygonsToCheck.Contains( ConnectedPolygonID ) )
-					{
-						PolygonsToCheck.Add( ConnectedPolygonID );
-					}
-				}
-			}
-			LastVertexID = VertexID;
-		}
-	}
-}
-
-
-void FMeshDescription::GetAllCharts( TArray<TArray<FPolygonID>>& OutCharts )
-{
-	// @todo: OutCharts: array of array doesn't seem like a really efficient data structure. Also templatize on allocator?
-
-	const int32 NumPolygons = Polygons().Num();
-
-	// Maintain a record of the polygons which have already been entered into a chart
-	// @todo: use TMemStackAllocator or similar to avoid expensive allocations
-	TSet<FPolygonID> ConsumedPolygons;
-	ConsumedPolygons.Reserve( NumPolygons );
-
-	for( const FPolygonID PolygonID : Polygons().GetElementIDs() )
-	{
-		if( !ConsumedPolygons.Contains( PolygonID ) )
-		{
-			TArray<FPolygonID>& Chart = OutCharts.AddDefaulted_GetRef();
-			GetPolygonsInSameChartAsPolygon( PolygonID, Chart );
-
-			// Mark all polygons in the chart as 'consumed'. Note that the chart will also contain the initial polygon.
-			for( const FPolygonID ChartPolygon : Chart )
-			{
-				ConsumedPolygons.Add( ChartPolygon );
-			}
-		}
-	}
-}
-
 
 void FMeshDescription::ReversePolygonFacing(const FPolygonID PolygonID)
 {
 	// Build a reverse perimeter
 	FMeshPolygon& Polygon = PolygonArray[PolygonID];
-	for (int32 i = 0; i < Polygon.PerimeterContour.VertexInstanceIDs.Num() / 2; ++i)
+	for (int32 i = 0; i < Polygon.VertexInstanceIDs.Num() / 2; ++i)
 	{
-		Polygon.PerimeterContour.VertexInstanceIDs.Swap(i, Polygon.PerimeterContour.VertexInstanceIDs.Num() - i - 1);
+		Polygon.VertexInstanceIDs.Swap(i, Polygon.VertexInstanceIDs.Num() - i - 1);
 	}
 
 	// Update the polygon's triangle vertex instance ids with the reversed ids
 	for (FTriangleID TriangleID : GetPolygonTriangleIDs(PolygonID))
 	{
 		FMeshTriangle& Triangle = TriangleArray[TriangleID];
-		FVertexInstanceID VertexInstanceID0 = Triangle.GetVertexInstanceID(0);
-		FVertexInstanceID VertexInstanceID2 = Triangle.GetVertexInstanceID(2);
-		Triangle.SetVertexInstanceID(0, VertexInstanceID2);
-		Triangle.SetVertexInstanceID(2, VertexInstanceID0);
+		Swap(Triangle.VertexInstanceIDs[0], Triangle.VertexInstanceIDs[2]);
 	}
 }
 
@@ -2015,7 +1195,6 @@ void FMeshDescription::RemapPolygonGroups(const TMap<FPolygonGroupID, FPolygonGr
 	}
 }
 
-PRAGMA_ENABLE_DEPRECATION_WARNINGS
 
 
 #if WITH_EDITORONLY_DATA
diff --git a/Engine/Source/Runtime/MeshDescription/Public/MeshDescription.h b/Engine/Source/Runtime/MeshDescription/Public/MeshDescription.h
--- a/Engine/Source/Runtime/MeshDescription/Public/MeshDescription.h
+++ b/Engine/Source/Runtime/MeshDescription/Public/MeshDescription.h
@@ -19,8 +19,6 @@
 #include "Containers/StaticArray.h"
 #include "MeshDescription.generated.h"
 
-PRAGMA_DISABLE_DEPRECATION_WARNINGS
-
 enum
 {
 	//Remove the _MD when FRawMesh will be remove
@@ -34,18 +32,20 @@ struct FMeshVertex
 	FMeshVertex()
 	{}
 
+private:
+
 	/** All of vertex instances which reference this vertex (for split vertex support) */
-	UE_DEPRECATED(4.24, "This member should not be accessed directly. Please use FMeshDescription::GetVertexVertexInstances instead.")
 	TArray<FVertexInstanceID> VertexInstanceIDs;
 
 	/** The edges connected to this vertex */
-	UE_DEPRECATED(4.24, "This member should not be accessed directly. Please use FMeshDescription::GetVertexConnectedEdges instead.")
 	TArray<FEdgeID> ConnectedEdgeIDs;
 
+public:
+
 	/** Serializer */
-	friend FArchive& operator<<( FArchive& Ar, FMeshVertex& Vertex )
+	friend FArchive& operator<<(FArchive& Ar, FMeshVertex& Vertex)
 	{
-		if( Ar.IsLoading() && Ar.CustomVer( FReleaseObjectVersion::GUID ) < FReleaseObjectVersion::MeshDescriptionNewSerialization )
+		if (Ar.IsLoading() && Ar.CustomVer(FReleaseObjectVersion::GUID) < FReleaseObjectVersion::MeshDescriptionNewSerialization)
 		{
 			Ar << Vertex.VertexInstanceIDs;
 			Ar << Vertex.ConnectedEdgeIDs;
@@ -61,30 +61,27 @@ struct FMeshVertexInstance
 	friend struct FMeshDescription;
 
 	FMeshVertexInstance()
-		: VertexID( FVertexID::Invalid )
+		: VertexID(FVertexID::Invalid)
 	{}
 
+private:
+
 	/** The vertex this is instancing */
-	UE_DEPRECATED(4.24, "This member should not be accessed directly. Please use FMeshDescription::GetVertexInstanceVertex instead.")
 	FVertexID VertexID;
 
-	/** List of connected polygons. This will soon be deprecated entirely, in favour of the ConnectedTriangles list. */
-	UE_DEPRECATED(4.24, "This member should not be accessed directly. Please use FMeshDescription::GetVertexInstanceConnectedPolygons instead.")
-	TArray<FPolygonID> ConnectedPolygons;
-
-private:
-	/** List of connected polygons */
+	/** List of connected triangles */
 	TArray<FTriangleID> ConnectedTriangles;
 
 public:
 
 	/** Serializer */
-	friend FArchive& operator<<( FArchive& Ar, FMeshVertexInstance& VertexInstance )
+	friend FArchive& operator<<(FArchive& Ar, FMeshVertexInstance& VertexInstance)
 	{
 		Ar << VertexInstance.VertexID;
-		if( Ar.IsLoading() && Ar.CustomVer( FReleaseObjectVersion::GUID ) < FReleaseObjectVersion::MeshDescriptionNewSerialization )
+		if (Ar.IsLoading() && Ar.CustomVer(FReleaseObjectVersion::GUID) < FReleaseObjectVersion::MeshDescriptionNewSerialization)
 		{
-			Ar << VertexInstance.ConnectedPolygons;
+			TArray<FPolygonID> ConnectedPolygons_DISCARD;
+			Ar << ConnectedPolygons_DISCARD;
 		}
 
 		return Ar;
@@ -98,33 +95,29 @@ struct FMeshEdge
 
 	FMeshEdge()
 	{
-		VertexIDs[ 0 ] = FVertexID::Invalid;
-		VertexIDs[ 1 ] = FVertexID::Invalid;
+		VertexIDs[0] = FVertexID::Invalid;
+		VertexIDs[1] = FVertexID::Invalid;
 	}
 
-	/** IDs of the two editable mesh vertices that make up this edge.  The winding direction is not defined. */
-	UE_DEPRECATED(4.24, "This member should not be accessed directly. Please use FMeshDescription::GetEdgeVertex instead.")
-	FVertexID VertexIDs[ 2 ];
+private:
 
-	/** The polygons that share this edge.  It's best if there are always only two polygons that share
-	    the edge, and those polygons are facing the same direction. This will soon be deprecated entirely, in favour of the ConnectedTriangles list. */
-	UE_DEPRECATED(4.24, "This member should not be accessed directly. Please use FMeshDescription::GetEdgeConnectedPolygons instead.")
-	TArray<FPolygonID> ConnectedPolygons;
+	/** IDs of the two editable mesh vertices that make up this edge.  The winding direction is not defined. */
+	FVertexID VertexIDs[2];
 
-private:
 	/** The triangles that share this edge */
 	TArray<FTriangleID> ConnectedTriangles;
 
 public:
 
 	/** Serializer */
-	friend FArchive& operator<<( FArchive& Ar, FMeshEdge& Edge )
+	friend FArchive& operator<<(FArchive& Ar, FMeshEdge& Edge)
 	{
-		Ar << Edge.VertexIDs[ 0 ];
-		Ar << Edge.VertexIDs[ 1 ];
-		if( Ar.IsLoading() && Ar.CustomVer( FReleaseObjectVersion::GUID ) < FReleaseObjectVersion::MeshDescriptionNewSerialization )
+		Ar << Edge.VertexIDs[0];
+		Ar << Edge.VertexIDs[1];
+		if (Ar.IsLoading() && Ar.CustomVer(FReleaseObjectVersion::GUID) < FReleaseObjectVersion::MeshDescriptionNewSerialization)
 		{
-			Ar << Edge.ConnectedPolygons;
+			TArray<FPolygonID> ConnectedPolygons_DISCARD;
+			Ar << ConnectedPolygons_DISCARD;
 		}
 
 		return Ar;
@@ -132,71 +125,50 @@ public:
 };
 
 
-struct UE_DEPRECATED(4.24, "This class should no longer be used.") FMeshPolygonContour
-{
-	/** The ordered list of vertex instances which make up the polygon contour. The winding direction is counter-clockwise. */
-	UE_DEPRECATED(4.24, "This member should not be accessed directly. Please use FMeshDescription::GetPolygonPerimeterVertexInstances instead.")
-	TArray<FVertexInstanceID> VertexInstanceIDs;
-
-	/** Serializer */
-	friend FArchive& operator<<( FArchive& Ar, FMeshPolygonContour& Contour )
-	{
-		Ar << Contour.VertexInstanceIDs;
-		return Ar;
-	}
-};
-
-
 struct FMeshTriangle
 {
 	friend struct FMeshDescription;
 
 	FMeshTriangle()
-		: VertexInstanceID0( FVertexInstanceID::Invalid ),
-		  VertexInstanceID1( FVertexInstanceID::Invalid ),
-		  VertexInstanceID2( FVertexInstanceID::Invalid )
-	{}
-
-	/** Vertex instance IDs that make up this triangle.  Indices must be ordered counter-clockwise. */
-	// @todo: when public member access is deprecated, we will replace these clumsy members with a simple array
-	UE_DEPRECATED(4.24, "Please access triangle vertex instances through Get/SetVertexInstanceID or FMeshDescription::GetTriangleVertexID.")
-	FVertexInstanceID VertexInstanceID0;
+	{
+		VertexInstanceIDs[0] = FVertexInstanceID::Invalid;
+		VertexInstanceIDs[1] = FVertexInstanceID::Invalid;
+		VertexInstanceIDs[2] = FVertexInstanceID::Invalid;
+	}
 
-	UE_DEPRECATED(4.24, "Please access triangle vertex instances through Get/SetVertexInstanceID or FMeshDescription::GetTriangleVertexID.")
-	FVertexInstanceID VertexInstanceID1;
+private:
 
-	UE_DEPRECATED(4.24, "Please access triangle vertex instances through Get/SetVertexInstanceID or FMeshDescription::GetTriangleVertexID.")
-	FVertexInstanceID VertexInstanceID2;
+	/** Vertex instance IDs that make up this triangle.  Indices must be ordered counter-clockwise. */
+	FVertexInstanceID VertexInstanceIDs[3];
 
-private:
 	/** Polygon which contains this triangle */
 	FPolygonID PolygonID;
 
 public:
+
 	/** Gets the specified triangle vertex instance ID.  Pass an index between 0 and 2 inclusive. */
-	inline FVertexInstanceID GetVertexInstanceID( const int32 Index ) const
+	inline FVertexInstanceID GetVertexInstanceID(const int32 Index) const
 	{
-		// When we deprecate direct member access, this will be a simple array lookup
-		checkSlow( Index >= 0 && Index <= 2 );
-		return ( &VertexInstanceID0 )[ Index ];
+		checkSlow(Index >= 0 && Index <= 2);
+		return VertexInstanceIDs[Index];
 	}
 
 	/** Sets the specified triangle vertex instance ID.  Pass an index between 0 and 2 inclusive, and the new vertex instance ID to store. */
-	inline void SetVertexInstanceID( const int32 Index, const FVertexInstanceID NewVertexInstanceID )
+	inline void SetVertexInstanceID(const int32 Index, const FVertexInstanceID NewVertexInstanceID)
 	{
 		// When we deprecate direct member access, this will be a simple array lookup
-		checkSlow( Index >= 0 && Index <= 2 );
-		( &VertexInstanceID0 )[ Index ] = NewVertexInstanceID;
+		checkSlow(Index >= 0 && Index <= 2);
+		VertexInstanceIDs[Index] = NewVertexInstanceID;
 	}
 
 	/** Serializer */
-	friend FArchive& operator<<( FArchive& Ar, FMeshTriangle& Triangle )
+	friend FArchive& operator<<(FArchive& Ar, FMeshTriangle& Triangle)
 	{
-		Ar << Triangle.VertexInstanceID0;
-		Ar << Triangle.VertexInstanceID1;
-		Ar << Triangle.VertexInstanceID2;
+		Ar << Triangle.VertexInstanceIDs[0];
+		Ar << Triangle.VertexInstanceIDs[1];
+		Ar << Triangle.VertexInstanceIDs[2];
 
-		if( !Ar.IsLoading() || Ar.CustomVer( FEditorObjectVersion::GUID ) >= FEditorObjectVersion::MeshDescriptionTriangles )
+		if (!Ar.IsLoading() || Ar.CustomVer(FEditorObjectVersion::GUID) >= FEditorObjectVersion::MeshDescriptionTriangles)
 		{
 			Ar << Triangle.PolygonID;
 		}
@@ -211,52 +183,50 @@ struct FMeshPolygon
 	friend struct FMeshDescription;
 
 	FMeshPolygon()
-		: PolygonGroupID( FPolygonGroupID::Invalid )
+		: PolygonGroupID(FPolygonGroupID::Invalid)
 	{}
 
-	/** The outer boundary edges of this polygon */
-	UE_DEPRECATED(4.24, "This member should not be accessed directly. Please use FMeshDescription::GetPolygonPerimeterVertexInstances instead.")
-	FMeshPolygonContour PerimeterContour;
+private:
 
-	/** List of triangles which make up this polygon */
-	UE_DEPRECATED(4.24, "This member is no longer used. Please use GetPolygonTriangleIDs instead.")
-	TArray<FMeshTriangle> Triangles;
+	/** The outer boundary edges of this polygon */
+	TArray<FVertexInstanceID> VertexInstanceIDs;
 
-private:
 	/** List of triangle IDs which make up this polygon */
 	TArray<FTriangleID> TriangleIDs;
 
-public:
 	/** The polygon group which contains this polygon */
-	UE_DEPRECATED(4.24, "This member should not be accessed directly. Please use FMeshDescription::GetPolygonPolygonGroup instead.")
 	FPolygonGroupID PolygonGroupID;
 
+public:
+
 	/** Serializer */
-	friend FArchive& operator<<( FArchive& Ar, FMeshPolygon& Polygon )
+	friend FArchive& operator<<(FArchive& Ar, FMeshPolygon& Polygon)
 	{
 		if (Ar.IsSaving() &&
 			Ar.CustomVer(FEditorObjectVersion::GUID) >= FEditorObjectVersion::MeshDescriptionTriangles &&
-			Polygon.PerimeterContour.VertexInstanceIDs.Num() == 3)
+			Polygon.VertexInstanceIDs.Num() == 3)
 		{
 			// Optimisation: if polygon is a triangle, don't serialize the vertices as they can be copied over from the associated triangle
-			FMeshPolygonContour Empty;
+			TArray<FVertexInstanceID> Empty;
 			Ar << Empty;
 		}
 		else
 		{
-			Ar << Polygon.PerimeterContour;
+			Ar << Polygon.VertexInstanceIDs;
 		}
 
-		if (Ar.IsLoading() && Ar.CustomVer( FEditorObjectVersion::GUID ) < FEditorObjectVersion::MeshDescriptionRemovedHoles)
+		if (Ar.IsLoading() && Ar.CustomVer(FEditorObjectVersion::GUID) < FEditorObjectVersion::MeshDescriptionRemovedHoles)
 		{
-			TArray<FMeshPolygonContour> Empty;
+			TArray<TArray<FVertexInstanceID>> Empty;
 			Ar << Empty;
 		}
-		if( Ar.IsLoading() && Ar.CustomVer( FReleaseObjectVersion::GUID ) < FReleaseObjectVersion::MeshDescriptionNewSerialization )
+
+		if (Ar.IsLoading() && Ar.CustomVer(FReleaseObjectVersion::GUID) < FReleaseObjectVersion::MeshDescriptionNewSerialization)
 		{
 			TArray<FMeshTriangle> Triangles_DISCARD;
 			Ar << Triangles_DISCARD;
 		}
+
 		Ar << Polygon.PolygonGroupID;
 
 		return Ar;
@@ -271,14 +241,17 @@ struct FMeshPolygonGroup
 	FMeshPolygonGroup()
 	{}
 
+private:
+
 	/** All polygons in this group */
-	UE_DEPRECATED(4.24, "This member should not be accessed directly. Please use FMeshDescription::GetPolygonGroupPolygons instead.")
 	TArray<FPolygonID> Polygons;
 
+public:
+
 	/** Serializer */
-	friend FArchive& operator<<( FArchive& Ar, FMeshPolygonGroup& PolygonGroup )
+	friend FArchive& operator<<(FArchive& Ar, FMeshPolygonGroup& PolygonGroup)
 	{
-		if( Ar.IsLoading() && Ar.CustomVer( FReleaseObjectVersion::GUID ) < FReleaseObjectVersion::MeshDescriptionNewSerialization )
+		if (Ar.IsLoading() && Ar.CustomVer(FReleaseObjectVersion::GUID) < FReleaseObjectVersion::MeshDescriptionNewSerialization)
 		{
 			Ar << PolygonGroup.Polygons;
 		}
@@ -382,46 +355,21 @@ public:
 	FVertexArray& Vertices() { return VertexArray; }
 	const FVertexArray& Vertices() const { return VertexArray; }
 
-	UE_DEPRECATED(4.24, "Please do not access the vertex directly - use FMeshDescription accessor methods.")
-	FMeshVertex& GetVertex(const FVertexID VertexID) { return VertexArray[VertexID]; }
-	UE_DEPRECATED(4.24, "Please do not access the vertex directly - use FMeshDescription accessor methods.")
-	const FMeshVertex& GetVertex(const FVertexID VertexID) const { return VertexArray[VertexID]; }
-
 	FVertexInstanceArray& VertexInstances() { return VertexInstanceArray; }
 	const FVertexInstanceArray& VertexInstances() const { return VertexInstanceArray; }
 
-	UE_DEPRECATED(4.24, "Please do not access the vertex instance directly - use FMeshDescription accessor methods.")
-	FMeshVertexInstance& GetVertexInstance(const FVertexInstanceID VertexInstanceID) { return VertexInstanceArray[VertexInstanceID]; }
-	UE_DEPRECATED(4.24, "Please do not access the vertex instance directly - use FMeshDescription accessor methods.")
-	const FMeshVertexInstance& GetVertexInstance(const FVertexInstanceID VertexInstanceID) const { return VertexInstanceArray[VertexInstanceID]; }
-
 	FEdgeArray& Edges() { return EdgeArray; }
 	const FEdgeArray& Edges() const { return EdgeArray; }
 
-	UE_DEPRECATED(4.24, "Please do not access the edge directly - use FMeshDescription accessor methods.")
-	FMeshEdge& GetEdge(const FEdgeID EdgeID) { return EdgeArray[EdgeID]; }
-	UE_DEPRECATED(4.24, "Please do not access the edge directly - use FMeshDescription accessor methods.")
-	const FMeshEdge& GetEdge(const FEdgeID EdgeID) const { return EdgeArray[EdgeID]; }
-
 	FTriangleArray& Triangles() { return TriangleArray; }
 	const FTriangleArray& Triangles() const { return TriangleArray; }
 
 	FPolygonArray& Polygons() { return PolygonArray; }
 	const FPolygonArray& Polygons() const { return PolygonArray; }
 
-	UE_DEPRECATED(4.24, "Please do not access the polygon directly - use FMeshDescription accessor methods.")
-	FMeshPolygon& GetPolygon(const FPolygonID PolygonID) { return PolygonArray[PolygonID]; }
-	UE_DEPRECATED(4.24, "Please do not access the polygon directly - use FMeshDescription accessor methods.")
-	const FMeshPolygon& GetPolygon(const FPolygonID PolygonID) const { return PolygonArray[PolygonID]; }
-
 	FPolygonGroupArray& PolygonGroups() { return PolygonGroupArray; }
 	const FPolygonGroupArray& PolygonGroups() const { return PolygonGroupArray; }
 
-	UE_DEPRECATED(4.24, "Please do not access the polygon group directly - use FMeshDescription accessor methods.")
-	FMeshPolygonGroup& GetPolygonGroup(const FPolygonGroupID PolygonGroupID) { return PolygonGroupArray[PolygonGroupID]; }
-	UE_DEPRECATED(4.24, "Please do not access the polygon group directly - use FMeshDescription accessor methods.")
-	const FMeshPolygonGroup& GetPolygonGroup(const FPolygonGroupID PolygonGroupID) const { return PolygonGroupArray[PolygonGroupID]; }
-
 	TAttributesSet<FVertexID>& VertexAttributes() { return VertexAttributesSet; }
 	const TAttributesSet<FVertexID>& VertexAttributes() const { return VertexAttributesSet; }
 
@@ -943,7 +891,7 @@ public:
 	/** Get the vertex instances which define this triangle */
 	TArrayView<const FVertexInstanceID> GetTriangleVertexInstances(const FTriangleID TriangleID) const
 	{
-		return TArrayView<const FVertexInstanceID>(&TriangleArray[TriangleID].VertexInstanceID0, 3);
+		return TriangleArray[TriangleID].VertexInstanceIDs;
 	}
 
 	/** Get the specified vertex instance by index */
@@ -1037,9 +985,6 @@ public:
 	//////////////////////////////////////////////////////////////////////
 	// Polygon operations
 
-	UE_DEPRECATED(4.24, "Please use GetPolygonTriangleIDs instead.")
-	TArray<FMeshTriangle> GetPolygonTriangles(const FPolygonID PolygonID) const;
-
 	/** Return reference to an array of triangle IDs which comprise this polygon */
 	const TArray<FTriangleID>& GetPolygonTriangleIDs(const FPolygonID PolygonID) const
 	{
@@ -1052,28 +997,16 @@ public:
 		return PolygonArray[PolygonID].TriangleIDs.Num();
 	}
 
-	UE_DEPRECATED(4.24, "Please use GetPolygonVertexInstances instead.")
-	const TArray<FVertexInstanceID>& GetPolygonPerimeterVertexInstances(const FPolygonID PolygonID) const
-	{
-		return GetPolygonVertexInstances(PolygonID);
-	}
-
 	/** Returns reference to an array of VertexInstance IDs forming the perimeter of this polygon */
 	const TArray<FVertexInstanceID>& GetPolygonVertexInstances(const FPolygonID PolygonID) const
 	{
-		return PolygonArray[PolygonID].PerimeterContour.VertexInstanceIDs;
+		return PolygonArray[PolygonID].VertexInstanceIDs;
 	}
 
 	/** Returns the number of vertices this polygon has */
 	int32 GetNumPolygonVertices(const FPolygonID PolygonID) const
 	{
-		return PolygonArray[PolygonID].PerimeterContour.VertexInstanceIDs.Num();
-	}
-
-	UE_DEPRECATED(4.24, "Please use GetPolygonVertices instead.")
-	void GetPolygonPerimeterVertices(const FPolygonID PolygonID, TArray<FVertexID>& OutVertexIDs) const
-	{
-		return GetPolygonVertices(PolygonID, OutVertexIDs);
+		return PolygonArray[PolygonID].VertexInstanceIDs.Num();
 	}
 
 	/** Populates the passed array of VertexIDs with the vertices which form the polygon perimeter */
@@ -1105,12 +1038,6 @@ public:
 		return Result;
 	}
 
-	UE_DEPRECATED(4.24, "Please use GetPolygonPerimeterEdges instead.")
-	void GetPolygonEdges(const FPolygonID PolygonID, TArray<FEdgeID>& OutEdgeIDs) const
-	{
-		GetPolygonPerimeterEdges(PolygonID, OutEdgeIDs);
-	}
-
 	/** Populates the passed array with the edges which form the polygon perimeter */
 	template <typename Alloc>
 	void GetPolygonPerimeterEdges(const FPolygonID PolygonID, TArray<FEdgeID, Alloc>& OutEdgeIDs) const
@@ -1184,7 +1111,7 @@ public:
 	/** Return the number of internal edges in this polygon */
 	int32 GetNumPolygonInternalEdges(const FPolygonID PolygonID) const
 	{
-		return PolygonArray[PolygonID].PerimeterContour.VertexInstanceIDs.Num() - 3;
+		return PolygonArray[PolygonID].VertexInstanceIDs.Num() - 3;
 	}
 
 	/** Populates the passed array with adjacent polygons */
@@ -1252,9 +1179,6 @@ public:
 	/** Reverse the winding order of the vertices of this polygon */
 	void ReversePolygonFacing(const FPolygonID PolygonID);
 
-	UE_DEPRECATED(4.24, "Please use the other overload of ComputePolygonTriangulation.")
-	void ComputePolygonTriangulation(const FPolygonID PolygonID, TArray<FMeshTriangle>& OutTriangles);
-
 	/** Generates triangles and internal edges for the given polygon */
 	void ComputePolygonTriangulation(const FPolygonID PolygonID);
 
@@ -1293,38 +1217,6 @@ public:
 	/** Retriangulates the entire mesh */
 	void TriangulateMesh();
 
-	/** Set the polygon tangent and normal only for the specified polygonIDs */
-	UE_DEPRECATED(4.24, "Please use FStaticMeshOperations::ComputePolygonTangentsAndNormals.")
-	void ComputePolygonTangentsAndNormals(const TArray<FPolygonID>& PolygonIDs, float ComparisonThreshold = 0.0f);
-
-	/** Set the polygon tangent and normal for all polygons in the mesh description. */
-	UE_DEPRECATED(4.24, "Please use FStaticMeshOperations::ComputePolygonTangentsAndNormals.")
-	void ComputePolygonTangentsAndNormals(float ComparisonThreshold = 0.0f);
-	
-	/** Set the vertex instance tangent and normal only for the specified VertexInstanceIDs */
-	UE_DEPRECATED(4.24, "Please use FStaticMeshOperations::ComputeTangentsAndNormals.")
-	void ComputeTangentsAndNormals(const TArray<FVertexInstanceID>& VertexInstanceIDs, EComputeNTBsOptions ComputeNTBsOptions);
-
-	/** Set the vertex instance tangent and normal for all vertex instances in the mesh description. */
-	UE_DEPRECATED(4.24, "Please use FStaticMeshOperations::ComputeTangentsAndNormals.")
-	void ComputeTangentsAndNormals(EComputeNTBsOptions ComputeNTBsOptions);
-
-	/** Determine the edge hardnesses from existing normals */
-	UE_DEPRECATED(4.24, "Please use FStaticMeshOperations::DetermineEdgeHardnessesFromVertexInstanceNormals.")
-	void DetermineEdgeHardnessesFromVertexInstanceNormals( const float Tolerance = KINDA_SMALL_NUMBER );
-
-	/** Determine UV seams from existing vertex instance UVs */
-	UE_DEPRECATED(4.24, "UVSeam attribute is now deprecated in MeshDescription.")
-	void DetermineUVSeamsFromUVs( const int32 UVIndex, const float Tolerance = KINDA_SMALL_NUMBER );
-
-	/** Get polygons in the same UV chart as the specified polygon */
-	UE_DEPRECATED(4.24, "UVSeam attribute is now deprecated in MeshDescription.")
-	void GetPolygonsInSameChartAsPolygon( const FPolygonID PolygonID, TArray<FPolygonID>& OutPolygonIDs );
-
-	/** Get array of all UV charts */
-	UE_DEPRECATED(4.24, "UVSeam attribute is now deprecated in MeshDescription.")
-	void GetAllCharts( TArray<TArray<FPolygonID>>& OutCharts );
-
 	/** Reverses the winding order of all polygons in the mesh */
 	void ReverseAllPolygonFacing();
 
@@ -1336,45 +1228,21 @@ private:
 
 	FPlane ComputePolygonPlane(const FPolygonID PolygonID) const;
 	FVector ComputePolygonNormal(const FPolygonID PolygonID) const;
-	bool ComputePolygonTangentsAndNormals(
-		  const FPolygonID PolygonID
-		, float ComparisonThreshold
-		, const TVertexAttributesRef<FVector> VertexPositions
-		, const TVertexInstanceAttributesRef<FVector2D> VertexUVs
-		, TPolygonAttributesRef<FVector> PolygonTangents
-		, TPolygonAttributesRef<FVector> PolygonBinormals
-		, TPolygonAttributesRef<FVector> PolygonCenters
-		, TPolygonAttributesRef<FVector> PolygonNormals
-	);
-
-	void GetVertexConnectedPolygonsInSameSoftEdgedGroup(const FVertexID VertexID, const FPolygonID PolygonID, TArray<FPolygonID>& OutPolygonIDs) const;
-	void GetPolygonsInSameSoftEdgedGroupAsPolygon(const FPolygonID PolygonID, const TArray<FPolygonID>& CandidatePolygonIDs, const TArray<FEdgeID>& SoftEdgeIDs, TArray<FPolygonID>& OutPolygonIDs) const;
-	void GetConnectedSoftEdges(const FVertexID VertexID, TArray<FEdgeID>& OutConnectedSoftEdges) const;
-	void ComputeTangentsAndNormals(
-		  const FVertexInstanceID VertexInstanceID
-		, EComputeNTBsOptions ComputeNTBsOptions
-		, const TPolygonAttributesRef<FVector> PolygonNormals
-		, const TPolygonAttributesRef<FVector> PolygonTangents
-		, const TPolygonAttributesRef<FVector> PolygonBinormals
-		, TVertexInstanceAttributesRef<FVector> VertexNormals
-		, TVertexInstanceAttributesRef<FVector> VertexTangents
-		, TVertexInstanceAttributesRef<float> VertexBinormalSigns
-	);
 
 private:
 
-	void CreateVertex_Internal( const FVertexID VertexID ) { VertexAttributesSet.Insert( VertexID ); }
-	void CreateVertexInstance_Internal( const FVertexInstanceID VertexInstanceID, const FVertexID VertexID );
-	void CreateEdge_Internal( const FEdgeID EdgeID, const FVertexID VertexID0, const FVertexID VertexID1 );
-	void CreateTriangle_Internal( const FTriangleID TriangleID, const FPolygonGroupID PolygonGroupID, TArrayView<const FVertexInstanceID> VertexInstanceIDs, TArray<FEdgeID>* OutEdgeIDs );
-	void CreatePolygon_Internal( const FPolygonID PolygonID, const FPolygonGroupID PolygonGroupID, TArrayView<const FVertexInstanceID> VertexInstanceIDs, TArray<FEdgeID>* OutEdgeIDs );
-	void CreatePolygonGroup_Internal( const FPolygonGroupID PolygonGroupID ) { PolygonGroupAttributesSet.Insert( PolygonGroupID ); }
+	void CreateVertex_Internal(const FVertexID VertexID) { VertexAttributesSet.Insert(VertexID); }
+	void CreateVertexInstance_Internal(const FVertexInstanceID VertexInstanceID, const FVertexID VertexID);
+	void CreateEdge_Internal(const FEdgeID EdgeID, const FVertexID VertexID0, const FVertexID VertexID1);
+	void CreateTriangle_Internal(const FTriangleID TriangleID, const FPolygonGroupID PolygonGroupID, TArrayView<const FVertexInstanceID> VertexInstanceIDs, TArray<FEdgeID>* OutEdgeIDs);
+	void CreatePolygon_Internal(const FPolygonID PolygonID, const FPolygonGroupID PolygonGroupID, TArrayView<const FVertexInstanceID> VertexInstanceIDs, TArray<FEdgeID>* OutEdgeIDs);
+	void CreatePolygonGroup_Internal(const FPolygonGroupID PolygonGroupID) { PolygonGroupAttributesSet.Insert(PolygonGroupID); }
 
 	/** Given a set of index remappings, fixes up references to element IDs */
-	void FixUpElementIDs( const FElementIDRemappings& Remappings );
+	void FixUpElementIDs(const FElementIDRemappings& Remappings);
 
 	/** Given a set of index remappings, remaps all attributes accordingly */
-	void RemapAttributes( const FElementIDRemappings& Remappings );
+	void RemapAttributes(const FElementIDRemappings& Remappings);
 
 
 	FVertexArray VertexArray;
@@ -1392,8 +1260,6 @@ private:
 	TAttributesSet<FPolygonGroupID> PolygonGroupAttributesSet;
 };
 
-PRAGMA_ENABLE_DEPRECATION_WARNINGS
-
 
 /**
  * Bulk data storage for FMeshDescription
@@ -1402,21 +1268,21 @@ struct MESHDESCRIPTION_API FMeshDescriptionBulkData
 {
 public:
 	FMeshDescriptionBulkData()
-		: bBulkDataUpdated( false )
+		: bBulkDataUpdated(false)
 		, bGuidIsHash(false)
 	{
-		BulkData.SetBulkDataFlags( BULKDATA_SerializeCompressed | BULKDATA_SerializeCompressedBitWindow );
+		BulkData.SetBulkDataFlags(BULKDATA_SerializeCompressed | BULKDATA_SerializeCompressedBitWindow);
 	}
 
 #if WITH_EDITORONLY_DATA
 	/** Serialization */
-	void Serialize( FArchive& Ar, UObject* Owner );
+	void Serialize(FArchive& Ar, UObject* Owner);
 
 	/** Store a new mesh description in the bulk data */
-	void SaveMeshDescription( FMeshDescription& MeshDescription );
+	void SaveMeshDescription(FMeshDescription& MeshDescription);
 
 	/** Load the mesh description from the bulk data */
-	void LoadMeshDescription( FMeshDescription& MeshDescription );
+	void LoadMeshDescription(FMeshDescription& MeshDescription);
 
 	/** Empties the bulk data */
 	void Empty();
@@ -1456,5 +1322,5 @@ public:
 	GENERATED_BODY()
 
 	// UObject interface
-	virtual void Serialize( FArchive& Ar ) override;
+	virtual void Serialize(FArchive& Ar) override;
 };
