Build fix

diff --git a/Engine/Source/Runtime/MeshDescription/Public/MeshDescription.h b/Engine/Source/Runtime/MeshDescription/Public/MeshDescription.h
--- a/Engine/Source/Runtime/MeshDescription/Public/MeshDescription.h
+++ b/Engine/Source/Runtime/MeshDescription/Public/MeshDescription.h
@@ -16,6 +16,7 @@
 #include "UObject/ReleaseObjectVersion.h"
 #include "Serialization/BulkData.h"
 #include "Serialization/CustomVersion.h"
+#include "Containers/StaticArray.h"
 #include "MeshDescription.generated.h"
 
 PRAGMA_DISABLE_DEPRECATION_WARNINGS
