First pass of MeshDescription API and format refactor.

diff --git a/Engine/Source/Runtime/MeshDescription/Private/MeshDescription.cpp b/Engine/Source/Runtime/MeshDescription/Private/MeshDescription.cpp
--- a/Engine/Source/Runtime/MeshDescription/Private/MeshDescription.cpp
+++ b/Engine/Source/Runtime/MeshDescription/Private/MeshDescription.cpp
@@ -8,26 +8,147 @@
 #include "Serialization/BulkDataWriter.h"
 #include "UObject/EnterpriseObjectVersion.h"
 
-void UDEPRECATED_MeshDescription::Serialize(FArchive& Ar)
+
+FName FMeshDescription::VerticesName("Vertices");
+FName FMeshDescription::VertexInstancesName("VertexInstances");
+FName FMeshDescription::UVsName("UVs");
+FName FMeshDescription::EdgesName("Edges");
+FName FMeshDescription::TrianglesName("Triangles");
+FName FMeshDescription::PolygonsName("Polygons");
+FName FMeshDescription::PolygonGroupsName("PolygonGroups");
+
+
+FMeshDescription::FMeshDescription()
+{
+	Initialize();
+}
+
+
+FMeshDescription::FMeshDescription(const FMeshDescription& Other)
 {
-	Super::Serialize(Ar);
+	Elements = Other.Elements;
+	Cache();
+}
 
-	if (!HasAnyFlags(RF_ClassDefaultObject))
+
+FMeshDescription& FMeshDescription::operator=(const FMeshDescription& Other)
+{
+	if (this != &Other)
 	{
-		UE_LOG(LogLoad, Error, TEXT( "UMeshDescription about to be deprecated - please resave %s"), *GetPathName());
+		Elements = Other.Elements;
+		Cache();
 	}
+	return *this;
+}
+
 
-	// Discard the contents
-	FMeshDescription MeshDescription;
-	Ar << MeshDescription;
+void FMeshDescription::Initialize()
+{
+	// Register the basic mesh element types
+	VertexElements = Elements.Emplace(VerticesName).Get();
+	VertexInstanceElements = Elements.Emplace(VertexInstancesName).Get();
+	UVElements = Elements.Emplace(UVsName).Get();
+	EdgeElements = Elements.Emplace(EdgesName).Get();
+	TriangleElements = Elements.Emplace(TrianglesName).Get();
+	PolygonElements = Elements.Emplace(PolygonsName).Get();
+	PolygonGroupElements = Elements.Emplace(PolygonGroupsName).Get();
+
+	// Now register topology-based attributes
+
+	// Register vertex reference for vertex instances
+	VertexInstanceVertices = VertexInstanceElements->Get().GetAttributes().RegisterIndexAttribute<FVertexID>(MeshAttribute::VertexInstance::VertexIndex);
+
+	// Register vertex references for edges
+	EdgeVertices = EdgeElements->Get().GetAttributes().RegisterIndexAttribute<FVertexID[2]>(MeshAttribute::Edge::VertexIndex);
+
+	// Register vertex instance and polygon references for triangles
+	TriangleVertexInstances = TriangleElements->Get().GetAttributes().RegisterIndexAttribute<FVertexInstanceID[3]>(MeshAttribute::Triangle::VertexInstanceIndex);
+	TrianglePolygons = TriangleElements->Get().GetAttributes().RegisterIndexAttribute<FPolygonID>(MeshAttribute::Triangle::PolygonIndex);
+	TrianglePolygonGroups = TriangleElements->Get().GetAttributes().RegisterIndexAttribute<FPolygonGroupID>(MeshAttribute::Triangle::PolygonGroupIndex);
+
+	// Register UV references for triangles
+	TriangleUVs = TriangleElements->Get().GetAttributes().RegisterIndexAttribute<FUVID[3]>(MeshAttribute::Triangle::UVIndex, 0);
+
+	// Register vertex and edge references for triangles; these are transient references which are generated at load time
+	TriangleVertices = TriangleElements->Get().GetAttributes().RegisterIndexAttribute<FVertexID[3]>(MeshAttribute::Triangle::VertexIndex, 1, EMeshAttributeFlags::IndexReference | EMeshAttributeFlags::Transient);
+	TriangleEdges = TriangleElements->Get().GetAttributes().RegisterIndexAttribute<FEdgeID[3]>(MeshAttribute::Triangle::EdgeIndex, 1, EMeshAttributeFlags::IndexReference | EMeshAttributeFlags::Transient);
+
+	// Register polygon group reference for polygons
+	PolygonPolygonGroups = PolygonElements->Get().GetAttributes().RegisterIndexAttribute<FPolygonGroupID>(MeshAttribute::Polygon::PolygonGroupIndex);
+
+	// Minimal requirement is that vertices have a Position attribute
+	VertexPositions = VertexElements->Get().GetAttributes().RegisterAttribute(MeshAttribute::Vertex::Position, 1, FVector::ZeroVector, EMeshAttributeFlags::Lerpable);
+
+	// Register UVCoordinates attribute for UVs
+	UVElements->Get().GetAttributes().RegisterAttribute(MeshAttribute::UV::UVCoordinate, 1, FVector2D::ZeroVector, EMeshAttributeFlags::Lerpable);
+
+	// Associate indexers with element types and their referencing attributes
+	InitializeIndexers();
 }
 
 
+void FMeshDescription::InitializeIndexers()
+{
+	// Vertices may typically have 6 vertex instances from the triangles which include them
+	VertexToVertexInstances.SetInitialNumReferences(6);
+	VertexToVertexInstances.Set(VertexElements, VertexInstanceElements, MeshAttribute::VertexInstance::VertexIndex);
+
+	// Vertices may typically have 6 edges from the triangles which include them
+	VertexToEdges.SetInitialNumReferences(6);
+	VertexToEdges.Set(VertexElements, EdgeElements, MeshAttribute::Edge::VertexIndex);
+
+	// @todo: VertexToTriangles?
+	VertexInstanceToTriangles.Set(VertexInstanceElements, TriangleElements, MeshAttribute::Triangle::VertexInstanceIndex);
+
+	// Assume most edges will have either 1 or 2 triangles
+	EdgeToTriangles.SetInitialNumReferences(2);
+	EdgeToTriangles.Set(EdgeElements, TriangleElements, MeshAttribute::Triangle::EdgeIndex);
+
+	// UVs may typically be used by 6 adjacent triangles
+	UVToTriangles.SetInitialNumReferences(6);
+	UVToTriangles.Set(UVElements, TriangleElements, MeshAttribute::Triangle::UVIndex);
+
+	// Assume most polygons are composed of only 1 triangle
+	PolygonToTriangles.SetInitialNumReferences(1);
+	PolygonToTriangles.Set(PolygonElements, TriangleElements, MeshAttribute::Triangle::PolygonIndex);
+
+	// Polygon group indexers are a little different; in general there will be many back references and few keys, so set not to use chunks.
+	PolygonGroupToTriangles.SetUnchunked();
+	PolygonGroupToPolygons.SetUnchunked();
+	PolygonGroupToTriangles.SetInitialNumReferences(256);
+	PolygonGroupToPolygons.SetInitialNumReferences(256);
+	PolygonGroupToTriangles.Set(PolygonGroupElements, TriangleElements, MeshAttribute::Triangle::PolygonGroupIndex);
+	PolygonGroupToPolygons.Set(PolygonGroupElements, PolygonElements, MeshAttribute::Polygon::PolygonGroupIndex);
+}
+
 
-FMeshDescription::FMeshDescription()
+void FMeshDescription::Cache()
 {
-	// Minimal requirement is that vertices have a Position attribute
-	VertexAttributesSet.RegisterAttribute(MeshAttribute::Vertex::Position, 1, FVector::ZeroVector, EMeshAttributeFlags::Lerpable);
+	// Get pointers to element containers
+	VertexElements = Elements.Find(VerticesName)->Get();
+	VertexInstanceElements = Elements.Find(VertexInstancesName)->Get();
+	UVElements = Elements.Find(UVsName)->Get();
+	EdgeElements = Elements.Find(EdgesName)->Get();
+	TriangleElements = Elements.Find(TrianglesName)->Get();
+	PolygonElements = Elements.Find(PolygonsName)->Get();
+	PolygonGroupElements = Elements.Find(PolygonGroupsName)->Get();
+
+	// Register required transient attributes
+	TriangleVertices = TriangleElements->Get().GetAttributes().RegisterIndexAttribute<FVertexID[3]>(MeshAttribute::Triangle::VertexIndex, 1, EMeshAttributeFlags::IndexReference | EMeshAttributeFlags::Transient);
+	TriangleEdges = TriangleElements->Get().GetAttributes().RegisterIndexAttribute<FEdgeID[3]>(MeshAttribute::Triangle::EdgeIndex, 1, EMeshAttributeFlags::IndexReference | EMeshAttributeFlags::Transient);
+
+	// Cache fundamental attribute arrays
+	VertexInstanceVertices = VertexInstanceElements->Get().GetAttributes().GetAttributesRef<FVertexID>(MeshAttribute::VertexInstance::VertexIndex);
+	EdgeVertices = EdgeElements->Get().GetAttributes().GetAttributesRef<TArrayView<FVertexID>>(MeshAttribute::Edge::VertexIndex);
+	TriangleVertexInstances = TriangleElements->Get().GetAttributes().GetAttributesRef<TArrayView<FVertexInstanceID>>(MeshAttribute::Triangle::VertexInstanceIndex);
+	TrianglePolygons = TriangleElements->Get().GetAttributes().GetAttributesRef<FPolygonID>(MeshAttribute::Triangle::PolygonIndex);
+	TrianglePolygonGroups = TriangleElements->Get().GetAttributes().GetAttributesRef<FPolygonGroupID>(MeshAttribute::Triangle::PolygonGroupIndex);
+	TriangleUVs = TriangleElements->Get().GetAttributes().GetAttributesRef<TArrayView<FUVID>>(MeshAttribute::Triangle::UVIndex);
+	PolygonPolygonGroups = PolygonElements->Get().GetAttributes().GetAttributesRef<FPolygonGroupID>(MeshAttribute::Polygon::PolygonGroupIndex);
+	VertexPositions = VertexElements->Get().GetAttributes().GetAttributesRef<FVector>(MeshAttribute::Vertex::Position);
+
+	// Associate indexers with element types and their referencing attributes
+	InitializeIndexers();
 }
 
 
@@ -41,6 +162,188 @@ void FMeshDescription::Serialize(FArchive& Ar)
 		UE_LOG(LogLoad, Warning, TEXT("Deprecated serialization format"));
 	}
 
+	if (Ar.IsLoading() && Ar.CustomVer(FEditorObjectVersion::GUID) < FReleaseObjectVersion::MeshDescriptionNewFormat)
+	{
+		// Serialize the old format data and transform it into the new format mesh element map
+		SerializeLegacy(Ar);
+	}
+	else
+	{
+		Ar << Elements;
+
+		// After loading elements, we need to re-cache mesh element and attribute arrays 
+		if (Ar.IsLoading())
+		{
+			Cache();
+
+			for (FTriangleID TriangleID : TriangleElements->Get().GetElementIDs())
+			{
+				for (int32 I = 0; I < 3; I++)
+				{
+					FVertexInstanceID VertexInstanceID = TriangleVertexInstances[TriangleID][I];
+					TriangleVertices[TriangleID][I] = VertexInstanceVertices[VertexInstanceID];
+					VertexInstanceToTriangles.AddReferenceToKey(VertexInstanceID, TriangleID);
+				}
+
+				for (int32 I = 0; I < 3; I++)
+				{
+					FEdgeID EdgeID = GetVertexPairEdge(TriangleVertices[TriangleID][I], TriangleVertices[TriangleID][(I+1) % 3]);
+					TriangleEdges[TriangleID][I] = EdgeID;
+					EdgeToTriangles.AddReferenceToKey(EdgeID, TriangleID);
+				}
+			}
+
+			RebuildIndexers();
+		}
+	}
+}
+
+
+struct FMeshVertex_Legacy
+{
+	TArray<FVertexInstanceID> VertexInstanceIDs;
+	TArray<FEdgeID> ConnectedEdgeIDs;
+
+	friend FArchive& operator<<(FArchive& Ar, FMeshVertex_Legacy& Vertex)
+	{
+		check(Ar.IsLoading());
+		if (Ar.CustomVer(FReleaseObjectVersion::GUID) < FReleaseObjectVersion::MeshDescriptionNewSerialization)
+		{
+			Ar << Vertex.VertexInstanceIDs;
+			Ar << Vertex.ConnectedEdgeIDs;
+		}
+
+		return Ar;
+	}
+};
+
+
+struct FMeshVertexInstance_Legacy
+{
+	FVertexID VertexID;
+	TArray<FTriangleID> ConnectedTriangles;
+
+	friend FArchive& operator<<(FArchive& Ar, FMeshVertexInstance_Legacy& VertexInstance)
+	{
+		check(Ar.IsLoading());
+		Ar << VertexInstance.VertexID;
+		if (Ar.CustomVer(FReleaseObjectVersion::GUID) < FReleaseObjectVersion::MeshDescriptionNewSerialization)
+		{
+			TArray<FPolygonID> ConnectedPolygons_DISCARD;
+			Ar << ConnectedPolygons_DISCARD;
+		}
+
+		return Ar;
+	}
+};
+
+
+struct FMeshEdge_Legacy
+{
+	FVertexID VertexIDs[2];
+	TArray<FTriangleID> ConnectedTriangles;
+
+	friend FArchive& operator<<(FArchive& Ar, FMeshEdge_Legacy& Edge)
+	{
+		check(Ar.IsLoading());
+		Ar << Edge.VertexIDs[0];
+		Ar << Edge.VertexIDs[1];
+		if (Ar.CustomVer(FReleaseObjectVersion::GUID) < FReleaseObjectVersion::MeshDescriptionNewSerialization)
+		{
+			TArray<FPolygonID> ConnectedPolygons_DISCARD;
+			Ar << ConnectedPolygons_DISCARD;
+		}
+
+		return Ar;
+	}
+};
+
+
+struct FMeshTriangle_Legacy
+{
+	FVertexInstanceID VertexInstanceIDs[3];
+	FPolygonID PolygonID;
+
+	friend FArchive& operator<<(FArchive& Ar, FMeshTriangle_Legacy& Triangle)
+	{
+		check(Ar.IsLoading());
+		Ar << Triangle.VertexInstanceIDs[0];
+		Ar << Triangle.VertexInstanceIDs[1];
+		Ar << Triangle.VertexInstanceIDs[2];
+
+		if (Ar.CustomVer(FEditorObjectVersion::GUID) >= FEditorObjectVersion::MeshDescriptionTriangles)
+		{
+			Ar << Triangle.PolygonID;
+		}
+
+		return Ar;
+	}
+};
+
+
+struct FMeshPolygon_Legacy
+{
+	TArray<FVertexInstanceID> VertexInstanceIDs;
+	TArray<FTriangleID> TriangleIDs;
+	FPolygonGroupID PolygonGroupID;
+
+	friend FArchive& operator<<(FArchive& Ar, FMeshPolygon_Legacy& Polygon)
+	{
+		check(Ar.IsLoading());
+		Ar << Polygon.VertexInstanceIDs;
+
+		if (Ar.CustomVer(FEditorObjectVersion::GUID) < FEditorObjectVersion::MeshDescriptionRemovedHoles)
+		{
+			TArray<TArray<FVertexInstanceID>> Empty;
+			Ar << Empty;
+		}
+
+		if (Ar.CustomVer(FReleaseObjectVersion::GUID) < FReleaseObjectVersion::MeshDescriptionNewSerialization)
+		{
+			TArray<FMeshTriangle_Legacy> Triangles_DISCARD;
+			Ar << Triangles_DISCARD;
+		}
+
+		Ar << Polygon.PolygonGroupID;
+
+		return Ar;
+	}
+};
+
+
+struct FMeshPolygonGroup_Legacy
+{
+	TArray<FPolygonID> Polygons;
+
+	friend FArchive& operator<<(FArchive& Ar, FMeshPolygonGroup_Legacy& PolygonGroup)
+	{
+		check(Ar.IsLoading());
+		if (Ar.CustomVer(FReleaseObjectVersion::GUID) < FReleaseObjectVersion::MeshDescriptionNewSerialization)
+		{
+			Ar << PolygonGroup.Polygons;
+		}
+
+		return Ar;
+	}
+};
+
+
+void FMeshDescription::SerializeLegacy(FArchive& Ar)
+{
+	TMeshElementArray<FMeshVertex_Legacy, FVertexID> VertexArray;
+	TMeshElementArray<FMeshVertexInstance_Legacy, FVertexInstanceID> VertexInstanceArray;
+	TMeshElementArray<FMeshEdge_Legacy, FEdgeID> EdgeArray;
+	TMeshElementArray<FMeshTriangle_Legacy, FTriangleID> TriangleArray;
+	TMeshElementArray<FMeshPolygon_Legacy, FPolygonID> PolygonArray;
+	TMeshElementArray<FMeshPolygonGroup_Legacy, FPolygonGroupID> PolygonGroupArray;
+
+	TAttributesSet<FVertexID> VertexAttributesSet;
+	TAttributesSet<FVertexInstanceID> VertexInstanceAttributesSet;
+	TAttributesSet<FEdgeID> EdgeAttributesSet;
+	TAttributesSet<FTriangleID> TriangleAttributesSet;
+	TAttributesSet<FPolygonID> PolygonAttributesSet;
+	TAttributesSet<FPolygonGroupID> PolygonGroupAttributesSet;
+
 	Ar << VertexArray;
 	Ar << VertexInstanceArray;
 	Ar << EdgeArray;
@@ -54,281 +357,267 @@ void FMeshDescription::Serialize(FArchive& Ar)
 	Ar << PolygonGroupAttributesSet;
 
 	// Serialize new triangle arrays since version MeshDescriptionTriangles
-	if (!Ar.IsLoading() || Ar.CustomVer(FEditorObjectVersion::GUID) >= FEditorObjectVersion::MeshDescriptionTriangles)
+	if (Ar.CustomVer(FEditorObjectVersion::GUID) >= FEditorObjectVersion::MeshDescriptionTriangles)
 	{
 		Ar << TriangleArray;
 		Ar << TriangleAttributesSet;
 	}
 
-	if (Ar.IsLoading() && Ar.CustomVer(FReleaseObjectVersion::GUID) >= FReleaseObjectVersion::MeshDescriptionNewSerialization)
+	// Convert the old style element arrays into the new format
+
+	// Completely reinitialize the mesh elements map as it is not being directly serialized into.
+	Initialize();
+
+	for (FVertexID VertexID : VertexArray.GetElementIDs())
+	{
+		VertexElements->Get().Insert(VertexID);
+	}
+
+	for (FVertexInstanceID VertexInstanceID : VertexInstanceArray.GetElementIDs())
 	{
-		// Populate vertex instance IDs for vertices
-		for (const FVertexInstanceID VertexInstanceID : VertexInstanceArray.GetElementIDs())
-		{
-			const FVertexID VertexID = GetVertexInstanceVertex(VertexInstanceID);
-			VertexArray[VertexID].VertexInstanceIDs.Add(VertexInstanceID);
-		}
+		VertexInstanceElements->Get().Insert(VertexInstanceID);
+		FVertexID VertexID = VertexInstanceArray[VertexInstanceID].VertexID;
+		VertexInstanceVertices[VertexInstanceID] = VertexID;
+		VertexToVertexInstances.AddReferenceToKey(VertexID, VertexInstanceID);
+	}
 
-		// Populate edge IDs for vertices
-		for (const FEdgeID EdgeID : EdgeArray.GetElementIDs())
+	for (FEdgeID EdgeID : EdgeArray.GetElementIDs())
+	{
+		EdgeElements->Get().Insert(EdgeID);
+		FVertexID VertexID0 = EdgeArray[EdgeID].VertexIDs[0];
+		FVertexID VertexID1 = EdgeArray[EdgeID].VertexIDs[1];
+		EdgeVertices[EdgeID][0] = VertexID0;
+		EdgeVertices[EdgeID][1] = VertexID1;
+		VertexToEdges.AddReferenceToKey(VertexID0, EdgeID);
+		VertexToEdges.AddReferenceToKey(VertexID1, EdgeID);
+	}
+
+	for (FPolygonGroupID PolygonGroupID : PolygonGroupArray.GetElementIDs())
+	{
+		PolygonGroupElements->Get().Insert(PolygonGroupID);
+	}
+
+	for (FPolygonID PolygonID : PolygonArray.GetElementIDs())
+	{
+		PolygonElements->Get().Insert(PolygonID);
+		FPolygonGroupID PolygonGroupID = PolygonArray[PolygonID].PolygonGroupID;
+		PolygonPolygonGroups[PolygonID] = PolygonGroupID;
+		PolygonGroupToPolygons.AddReferenceToKey(PolygonGroupID, PolygonID);
+
+		// If the asset is pre-triangles, we generate triangles here from the polygon
+		if (Ar.CustomVer(FEditorObjectVersion::GUID) < FEditorObjectVersion::MeshDescriptionTriangles)
 		{
-			const FVertexID VertexID0 = GetEdgeVertex(EdgeID, 0);
-			const FVertexID VertexID1 = GetEdgeVertex(EdgeID, 1);
-			VertexArray[VertexID0].ConnectedEdgeIDs.Add(EdgeID);
-			VertexArray[VertexID1].ConnectedEdgeIDs.Add(EdgeID);
+			check(PolygonArray[PolygonID].VertexInstanceIDs.Num() >= 3);
+			CreatePolygonTriangles(PolygonID, PolygonArray[PolygonID].VertexInstanceIDs);
 		}
+	}
 
-		if (Ar.CustomVer(FEditorObjectVersion::GUID) >= FEditorObjectVersion::MeshDescriptionTriangles)
+	// Only do this if there were actually triangles in the asset
+	if (Ar.CustomVer(FEditorObjectVersion::GUID) >= FEditorObjectVersion::MeshDescriptionTriangles)
+	{
+		for (FTriangleID TriangleID : TriangleArray.GetElementIDs())
 		{
-			// Make reverse connection from polygons to triangles
-			for (const FTriangleID TriangleID : TriangleArray.GetElementIDs())
+			TriangleElements->Get().Insert(TriangleID);
+			FPolygonID PolygonID = TriangleArray[TriangleID].PolygonID;
+			FPolygonGroupID PolygonGroupID = PolygonArray[PolygonID].PolygonGroupID;
+			TrianglePolygons[TriangleID] = PolygonID;
+			TrianglePolygonGroups[TriangleID] = PolygonGroupID;
+			PolygonToTriangles.AddReferenceToKey(PolygonID, TriangleID);
+			PolygonGroupToTriangles.AddReferenceToKey(PolygonGroupID, TriangleID);
+
+			for (int32 I = 0; I < 3; I++)
 			{
-				const FPolygonID PolygonID = TriangleArray[TriangleID].PolygonID;
-				PolygonArray[PolygonID].TriangleIDs.Add(TriangleID);
+				FVertexInstanceID VertexInstanceID = TriangleArray[TriangleID].VertexInstanceIDs[I];
+				TriangleVertexInstances[TriangleID][I] = VertexInstanceID;
+				TriangleVertices[TriangleID][I] = VertexInstanceVertices[VertexInstanceID];
+				VertexInstanceToTriangles.AddReferenceToKey(VertexInstanceID, TriangleID);
 			}
-		}
 
-		// Populate polygon IDs for vertex instances, edges and polygon groups
-		for (const FPolygonID PolygonID : PolygonArray.GetElementIDs())
-		{
-			if (Ar.CustomVer(FEditorObjectVersion::GUID) >= FEditorObjectVersion::MeshDescriptionTriangles)
+			for (int32 I = 0; I < 3; I++)
 			{
-				// If the polygon has no contour serialized, copy it over from the triangle
-				if (PolygonArray[PolygonID].VertexInstanceIDs.Num() == 0)
-				{
-					check(PolygonArray[PolygonID].TriangleIDs.Num() == 1);
-					const FTriangleID TriangleID = PolygonArray[PolygonID].TriangleIDs[0];
-					for (int32 Index = 0; Index < 3; Index++)
-					{
-						PolygonArray[PolygonID].VertexInstanceIDs.Add(TriangleArray[TriangleID].GetVertexInstanceID(Index));
-					}
-				}
+				FEdgeID EdgeID = GetVertexPairEdge(TriangleVertices[TriangleID][I], TriangleVertices[TriangleID][(I+1) % 3]);
+				TriangleEdges[TriangleID][I] = EdgeID;
+				EdgeToTriangles.AddReferenceToKey(EdgeID, TriangleID);
 			}
-
-			const FPolygonGroupID PolygonGroupID = PolygonArray[PolygonID].PolygonGroupID;
-			PolygonGroupArray[PolygonGroupID].Polygons.Add(PolygonID);
 		}
 	}
 
-	if (Ar.IsLoading())
+	// Unregister the position attribute from the vertex elements, as it will come from the legacy data instead.
+	VertexElements->Get().GetAttributes().UnregisterAttribute(MeshAttribute::Vertex::Position);
+
+	// Add the legacy mesh attributes into the new containers
+	VertexElements->Get().GetAttributes().AppendAttributesFrom(VertexAttributesSet);
+	VertexInstanceElements->Get().GetAttributes().AppendAttributesFrom(VertexInstanceAttributesSet);
+	EdgeElements->Get().GetAttributes().AppendAttributesFrom(EdgeAttributesSet);
+	PolygonElements->Get().GetAttributes().AppendAttributesFrom(PolygonAttributesSet);
+	PolygonGroupElements->Get().GetAttributes().AppendAttributesFrom(PolygonGroupAttributesSet);
+
+	if (Ar.CustomVer(FEditorObjectVersion::GUID) >= FEditorObjectVersion::MeshDescriptionTriangles)
 	{
-		if (Ar.CustomVer(FEditorObjectVersion::GUID) < FEditorObjectVersion::MeshDescriptionTriangles)
-		{
-			TriangleArray.Reset();
+		TriangleElements->Get().GetAttributes().AppendAttributesFrom(TriangleAttributesSet);
+	}
 
-			// If we didn't serialize triangles, generate them from the polygon contour
-			for (const FPolygonID PolygonID : PolygonArray.GetElementIDs())
-			{
-				check(PolygonArray[PolygonID].TriangleIDs.Num() == 0);
-				ComputePolygonTriangulation(PolygonID);
-			}
-		}
-		else
-		{
-			// Otherwise connect existing triangles to vertex instances and edges
-			for (const FTriangleID TriangleID : TriangleArray.GetElementIDs())
-			{
-				for (int32 Index = 0; Index < 3; ++Index)
-				{
-					const FVertexInstanceID VertexInstanceID = GetTriangleVertexInstance(TriangleID, Index);
-					const FVertexInstanceID NextVertexInstanceID = GetTriangleVertexInstance(TriangleID, (Index + 1 == 3) ? 0 : Index + 1);
+	Cache();
+	BuildIndexers();
+}
 
-					const FVertexID VertexID0 = GetVertexInstanceVertex(VertexInstanceID);
-					const FVertexID VertexID1 = GetVertexInstanceVertex(NextVertexInstanceID);
 
-					FEdgeID EdgeID = GetVertexPairEdge(VertexID0, VertexID1);
-					check(EdgeID != FEdgeID::Invalid);
+void FMeshDescription::ResetIndexers()
+{
+	VertexToVertexInstances.Reset();
+	VertexToEdges.Reset();
+	VertexInstanceToTriangles.Reset();
+	EdgeToTriangles.Reset();
+	UVToTriangles.Reset();
+	PolygonToTriangles.Reset();
+	PolygonGroupToTriangles.Reset();
+	PolygonGroupToPolygons.Reset();
+}
 
-					VertexInstanceArray[VertexInstanceID].ConnectedTriangles.Add(TriangleID);
-					EdgeArray[EdgeID].ConnectedTriangles.Add(TriangleID);
-				}
-			}
-		}
-	}
+
+void FMeshDescription::BuildIndexers()
+{
+	VertexToVertexInstances.Build();
+	VertexToEdges.Build();
+	VertexInstanceToTriangles.Build();
+	EdgeToTriangles.Build();
+	UVToTriangles.Build();
+	PolygonToTriangles.Build();
+	PolygonGroupToTriangles.Build();
+	PolygonGroupToPolygons.Build();
+}
+
+
+void FMeshDescription::RebuildIndexers()
+{
+	VertexToVertexInstances.ForceRebuild();
+	VertexToEdges.ForceRebuild();
+	VertexInstanceToTriangles.ForceRebuild();
+	EdgeToTriangles.ForceRebuild();
+	UVToTriangles.ForceRebuild();
+	PolygonToTriangles.ForceRebuild();
+	PolygonGroupToTriangles.ForceRebuild();
+	PolygonGroupToPolygons.ForceRebuild();
 }
 
 
 void FMeshDescription::Empty()
 {
-	VertexArray.Reset();
-	VertexInstanceArray.Reset();
-	EdgeArray.Reset();
-	TriangleArray.Reset();
-	PolygonArray.Reset();
-	PolygonGroupArray.Reset();
-
-	// Empty all attributes
-	VertexAttributesSet.Initialize(0);
-	VertexInstanceAttributesSet.Initialize(0);
-	EdgeAttributesSet.Initialize(0);
-	TriangleAttributesSet.Initialize(0);
-	PolygonAttributesSet.Initialize(0);
-	PolygonGroupAttributesSet.Initialize(0);
+	VertexElements->Reset();
+	VertexInstanceElements->Reset();
+	UVElements->Reset();
+	EdgeElements->Reset();
+	TriangleElements->Reset();
+	PolygonElements->Reset();
+	PolygonGroupElements->Reset();
+	ResetIndexers();
 }
 
 
 bool FMeshDescription::IsEmpty() const
 {
-	return VertexArray.GetArraySize() == 0 &&
-		   VertexInstanceArray.GetArraySize() == 0 &&
-		   EdgeArray.GetArraySize() == 0 &&
-		   TriangleArray.GetArraySize() == 0 &&
-		   PolygonArray.GetArraySize() == 0 &&
-		   PolygonGroupArray.GetArraySize() == 0;
+	return VertexElements->IsEmpty() &&
+		   VertexInstanceElements->IsEmpty() &&
+		   UVElements->IsEmpty() &&
+		   EdgeElements->IsEmpty() &&
+		   TriangleElements->IsEmpty() &&
+		   PolygonElements->IsEmpty() &&
+		   PolygonGroupElements->IsEmpty();
 }
 
 
 void FMeshDescription::Compact(FElementIDRemappings& OutRemappings)
 {
-	VertexArray.Compact(OutRemappings.NewVertexIndexLookup);
-	VertexInstanceArray.Compact(OutRemappings.NewVertexInstanceIndexLookup);
-	EdgeArray.Compact(OutRemappings.NewEdgeIndexLookup);
-	TriangleArray.Compact(OutRemappings.NewTriangleIndexLookup);
-	PolygonArray.Compact(OutRemappings.NewPolygonIndexLookup);
-	PolygonGroupArray.Compact(OutRemappings.NewPolygonGroupIndexLookup);
+	VertexElements->Get().Compact(OutRemappings.NewVertexIndexLookup);
+	VertexInstanceElements->Get().Compact(OutRemappings.NewVertexInstanceIndexLookup);
+	EdgeElements->Get().Compact(OutRemappings.NewEdgeIndexLookup);
+	TriangleElements->Get().Compact(OutRemappings.NewTriangleIndexLookup);
+	PolygonElements->Get().Compact(OutRemappings.NewPolygonIndexLookup);
+	PolygonGroupElements->Get().Compact(OutRemappings.NewPolygonGroupIndexLookup);
 
-	RemapAttributes(OutRemappings);
 	FixUpElementIDs(OutRemappings);
 }
 
 
 void FMeshDescription::Remap(const FElementIDRemappings& Remappings)
 {
-	VertexArray.Remap(Remappings.NewVertexIndexLookup);
-	VertexInstanceArray.Remap(Remappings.NewVertexInstanceIndexLookup);
-	EdgeArray.Remap(Remappings.NewEdgeIndexLookup);
-	TriangleArray.Remap(Remappings.NewTriangleIndexLookup);
-	PolygonArray.Remap(Remappings.NewPolygonIndexLookup);
-	PolygonGroupArray.Remap(Remappings.NewPolygonGroupIndexLookup);
+	VertexElements->Get().Remap(Remappings.NewVertexIndexLookup);
+	VertexInstanceElements->Get().Remap(Remappings.NewVertexInstanceIndexLookup);
+	EdgeElements->Get().Remap(Remappings.NewEdgeIndexLookup);
+	TriangleElements->Get().Remap(Remappings.NewTriangleIndexLookup);
+	PolygonElements->Get().Remap(Remappings.NewPolygonIndexLookup);
+	PolygonGroupElements->Get().Remap(Remappings.NewPolygonGroupIndexLookup);
 
-	RemapAttributes(Remappings);
 	FixUpElementIDs(Remappings);
 }
 
 
-void FMeshDescription::RemapAttributes(const FElementIDRemappings& Remappings)
-{
-	VertexAttributesSet.Remap(Remappings.NewVertexIndexLookup);
-	VertexInstanceAttributesSet.Remap(Remappings.NewVertexInstanceIndexLookup);
-	EdgeAttributesSet.Remap(Remappings.NewEdgeIndexLookup);
-	TriangleAttributesSet.Remap(Remappings.NewTriangleIndexLookup);
-	PolygonAttributesSet.Remap(Remappings.NewPolygonIndexLookup);
-	PolygonGroupAttributesSet.Remap(Remappings.NewPolygonGroupIndexLookup);
-}
-
-
 void FMeshDescription::FixUpElementIDs(const FElementIDRemappings& Remappings)
 {
-	for (const FVertexID VertexID : VertexArray.GetElementIDs())
-	{
-		FMeshVertex& Vertex = VertexArray[VertexID];
-
-		// Fix up vertex instance index references in vertices array
-		for (FVertexInstanceID& VertexInstanceID : Vertex.VertexInstanceIDs)
-		{
-			VertexInstanceID = Remappings.GetRemappedVertexInstanceID(VertexInstanceID);
-		}
-
-		// Fix up edge index references in the vertex array
-		for (FEdgeID& EdgeID : Vertex.ConnectedEdgeIDs)
-		{
-			EdgeID = Remappings.GetRemappedEdgeID(EdgeID);
-		}
-	}
-
 	// Fix up vertex index references in vertex instance array
-	for (const FVertexInstanceID VertexInstanceID : VertexInstanceArray.GetElementIDs())
+	for (const FVertexInstanceID VertexInstanceID : VertexInstanceElements->Get().GetElementIDs())
 	{
-		FMeshVertexInstance& VertexInstance = VertexInstanceArray[VertexInstanceID];
-
-		VertexInstance.VertexID = Remappings.GetRemappedVertexID(VertexInstance.VertexID);
-
-		for (FTriangleID& TriangleID : VertexInstance.ConnectedTriangles)
-		{
-			TriangleID = Remappings.GetRemappedTriangleID(TriangleID);
-		}
+		VertexInstanceVertices[VertexInstanceID] = Remappings.GetRemappedVertexID(VertexInstanceVertices[VertexInstanceID]);
 	}
 
-	for (const FEdgeID EdgeID : EdgeArray.GetElementIDs())
+	for (const FEdgeID EdgeID : EdgeElements->Get().GetElementIDs())
 	{
-		FMeshEdge& Edge = EdgeArray[EdgeID];
-
 		// Fix up vertex index references in Edges array
 		for (int32 Index = 0; Index < 2; Index++)
 		{
-			Edge.VertexIDs[Index] = Remappings.GetRemappedVertexID(Edge.VertexIDs[Index]);
-		}
-
-		for (FTriangleID& TriangleID : Edge.ConnectedTriangles)
-		{
-			TriangleID = Remappings.GetRemappedTriangleID(TriangleID);
+			EdgeVertices[EdgeID][Index] = Remappings.GetRemappedVertexID(EdgeVertices[EdgeID][Index]);
 		}
 	}
 
-	for (const FTriangleID TriangleID : TriangleArray.GetElementIDs())
+	for (const FTriangleID TriangleID : TriangleElements->Get().GetElementIDs())
 	{
-		FMeshTriangle& Triangle = TriangleArray[TriangleID];
-		
 		// Fix up vertex instance references in Triangle
-		for (FVertexInstanceID& VertexInstanceID : Triangle.VertexInstanceIDs)
+		for (int32 Index = 0; Index < 3; Index++)
 		{
-			VertexInstanceID = Remappings.GetRemappedVertexInstanceID(VertexInstanceID);
+			TriangleVertexInstances[TriangleID][Index] = Remappings.GetRemappedVertexInstanceID(TriangleVertexInstances[TriangleID][Index]);
+			TriangleEdges[TriangleID][Index] = Remappings.GetRemappedEdgeID(TriangleEdges[TriangleID][Index]);
+			TriangleVertices[TriangleID][Index] = Remappings.GetRemappedVertexID(TriangleVertices[TriangleID][Index]);
 		}
 
-		Triangle.PolygonID = Remappings.GetRemappedPolygonID(Triangle.PolygonID);
+		TrianglePolygons[TriangleID] = Remappings.GetRemappedPolygonID(TrianglePolygons[TriangleID]);
+		TrianglePolygonGroups[TriangleID] = Remappings.GetRemappedPolygonGroupID(TrianglePolygonGroups[TriangleID]);
 	}
 
-	for (const FPolygonID PolygonID : PolygonArray.GetElementIDs())
+	for (const FPolygonID PolygonID : PolygonElements->Get().GetElementIDs())
 	{
-		FMeshPolygon& Polygon = PolygonArray[PolygonID];
-
-		// Fix up references to vertex indices in section polygons' contours
-		for (FVertexInstanceID& VertexInstanceID : Polygon.VertexInstanceIDs)
-		{
-			VertexInstanceID = Remappings.GetRemappedVertexInstanceID(VertexInstanceID);
-		}
-
-		for (FTriangleID& TriangleID : Polygon.TriangleIDs)
-		{
-			TriangleID = Remappings.GetRemappedTriangleID(TriangleID);
-		}
-
-		Polygon.PolygonGroupID = Remappings.GetRemappedPolygonGroupID(Polygon.PolygonGroupID);
+		PolygonPolygonGroups[PolygonID] = Remappings.GetRemappedPolygonGroupID(PolygonPolygonGroups[PolygonID]);
 	}
 
-	for (const FPolygonGroupID PolygonGroupID : PolygonGroupArray.GetElementIDs())
-	{
-		FMeshPolygonGroup& PolygonGroup = PolygonGroupArray[PolygonGroupID];
-
-		for (FPolygonID& Polygon : PolygonGroup.Polygons)
-		{
-			Polygon = Remappings.GetRemappedPolygonID(Polygon);
-		}
-	}
+	RebuildIndexers();
 }
 
 
 void FMeshDescription::CreateVertexInstance_Internal(const FVertexInstanceID VertexInstanceID, const FVertexID VertexID)
 {
-	VertexInstanceArray[VertexInstanceID].VertexID = VertexID;
-	check(!VertexArray[VertexID].VertexInstanceIDs.Contains(VertexInstanceID));
-	VertexArray[VertexID].VertexInstanceIDs.Add(VertexInstanceID);
-	VertexInstanceAttributesSet.Insert(VertexInstanceID);
+	VertexInstanceVertices[VertexInstanceID] = VertexID;
+	VertexToVertexInstances.AddReferenceToKey(VertexID, VertexInstanceID);
 }
 
 
 template <template <typename...> class TContainer>
 void FMeshDescription::DeleteVertexInstance_Internal(const FVertexInstanceID VertexInstanceID, TContainer<FVertexID>* InOutOrphanedVerticesPtr)
 {
-	check(VertexInstanceArray[VertexInstanceID].ConnectedTriangles.Num() == 0);
-	const FVertexID VertexID = VertexInstanceArray[VertexInstanceID].VertexID;
-	verify(VertexArray[VertexID].VertexInstanceIDs.RemoveSingle(VertexInstanceID) == 1);
-	if (InOutOrphanedVerticesPtr && VertexArray[VertexID].VertexInstanceIDs.Num() == 0 && VertexArray[VertexID].ConnectedEdgeIDs.Num() == 0)
+	checkSlow(VertexInstanceToTriangles.Find(VertexInstanceID).Num() == 0);
+
+	const FVertexID VertexID = VertexInstanceVertices[VertexInstanceID];
+	VertexToVertexInstances.RemoveReferenceFromKey(VertexID, VertexInstanceID);
+
+	VertexInstanceElements->Get().Remove(VertexInstanceID);
+	VertexInstanceToTriangles.RemoveKey(VertexInstanceID);
+
+	// Always perform the Find() after the element has been completely deleted, so it won't be included in reindexing if the key is stale.
+	// Note: removing the element above will clear the attributes to default, so remember the VertexID we're interested in.
+	if (InOutOrphanedVerticesPtr &&
+		VertexToVertexInstances.Find(VertexID).Num() == 0 &&
+		VertexToEdges.Find(VertexID).Num() == 0)
 	{
 		AddUnique(*InOutOrphanedVerticesPtr, VertexID);
 	}
-	VertexInstanceArray.Remove(VertexInstanceID);
-	VertexInstanceAttributesSet.Remove(VertexInstanceID);
 }
 
 void FMeshDescription::DeleteVertexInstance(const FVertexInstanceID VertexInstanceID, TArray<FVertexID>* InOutOrphanedVerticesPtr)
@@ -338,31 +627,44 @@ void FMeshDescription::DeleteVertexInstance(const FVertexInstanceID VertexInstan
 
 void FMeshDescription::CreateEdge_Internal(const FEdgeID EdgeID, const FVertexID VertexID0, const FVertexID VertexID1)
 {
-	check(GetVertexPairEdge(VertexID0, VertexID1) == FEdgeID::Invalid);
-	FMeshEdge& Edge = EdgeArray[EdgeID];
-	Edge.VertexIDs[0] = VertexID0;
-	Edge.VertexIDs[1] = VertexID1;
-	VertexArray[VertexID0].ConnectedEdgeIDs.Add(EdgeID);
-	VertexArray[VertexID1].ConnectedEdgeIDs.Add(EdgeID);
-	EdgeAttributesSet.Insert(EdgeID);
+	checkSlow(GetVertexPairEdge(VertexID0, VertexID1) == INDEX_NONE);
+
+	TArrayView<FVertexID> EdgeVertexIDs = EdgeVertices[EdgeID];
+	EdgeVertexIDs[0] = VertexID0;
+	EdgeVertexIDs[1] = VertexID1;
+
+	VertexToEdges.AddReferenceToKey(VertexID0, EdgeID);
+	VertexToEdges.AddReferenceToKey(VertexID1, EdgeID);
 }
 
 template <template <typename...> class TContainer>
 void FMeshDescription::DeleteEdge_Internal(const FEdgeID EdgeID, TContainer<FVertexID>* InOutOrphanedVerticesPtr)
 {
-	FMeshEdge& Edge = EdgeArray[EdgeID];
-	for (const FVertexID EdgeVertexID : Edge.VertexIDs)
+	const FVertexID VertexID0 = EdgeVertices[EdgeID][0];
+	const FVertexID VertexID1 = EdgeVertices[EdgeID][1];
+
+	VertexToEdges.RemoveReferenceFromKey(VertexID0, EdgeID);
+	VertexToEdges.RemoveReferenceFromKey(VertexID1, EdgeID);
+
+	EdgeElements->Get().Remove(EdgeID);
+	EdgeToTriangles.RemoveKey(EdgeID);
+
+	// Always perform the Find() after the element has been completely deleted, so it won't be included in reindexing if the key is stale.
+	// Note: removing the element above will clear the attributes to default, so remember the VertexID we're interested in.
+	if (InOutOrphanedVerticesPtr)
 	{
-		FMeshVertex& Vertex = VertexArray[EdgeVertexID];
-		verify(Vertex.ConnectedEdgeIDs.RemoveSingle(EdgeID) == 1);
-		if (InOutOrphanedVerticesPtr && Vertex.ConnectedEdgeIDs.Num() == 0)
+		if (VertexToEdges.Find(VertexID0).Num() == 0)
 		{
-			check(Vertex.VertexInstanceIDs.Num() == 0);  // We must already have deleted any vertex instances
-			AddUnique(*InOutOrphanedVerticesPtr, EdgeVertexID);
+			check(VertexToVertexInstances.Find(VertexID0).Num() == 0);
+			AddUnique(*InOutOrphanedVerticesPtr, VertexID0);
+		}
+
+		if (VertexToEdges.Find(VertexID1).Num() == 0)
+		{
+			check(VertexToVertexInstances.Find(VertexID1).Num() == 0);
+			AddUnique(*InOutOrphanedVerticesPtr, VertexID1);
 		}
 	}
-	EdgeArray.Remove(EdgeID);
-	EdgeAttributesSet.Remove(EdgeID);
 }
 
 void FMeshDescription::DeleteEdge(const FEdgeID EdgeID, TArray<FVertexID>* InOutOrphanedVerticesPtr)
@@ -378,38 +680,41 @@ void FMeshDescription::CreateTriangle_Internal(const FTriangleID TriangleID, con
 	}
 
 	// Fill out triangle vertex instances
-	FMeshTriangle& Triangle = TriangleArray[TriangleID];
+	TArrayView<FVertexInstanceID> TriVertexInstanceIDs = TriangleVertexInstances[TriangleID];
 	check(VertexInstanceIDs.Num() == 3);
-	Triangle.SetVertexInstanceID(0, VertexInstanceIDs[0]);
-	Triangle.SetVertexInstanceID(1, VertexInstanceIDs[1]);
-	Triangle.SetVertexInstanceID(2, VertexInstanceIDs[2]);
+	TriVertexInstanceIDs[0] = VertexInstanceIDs[0];
+	TriVertexInstanceIDs[1] = VertexInstanceIDs[1];
+	TriVertexInstanceIDs[2] = VertexInstanceIDs[2];
 
-	// Make a polygon which will contain this triangle
-	FPolygonID PolygonID = PolygonArray.Add();
-	FMeshPolygon& Polygon = PolygonArray[PolygonID];
-	PolygonAttributesSet.Insert(PolygonID);
+	// Fill out triangle polygon group
+	TrianglePolygonGroups[TriangleID] = PolygonGroupID;
+	PolygonGroupToTriangles.AddReferenceToKey(PolygonGroupID, TriangleID);
 
-	Polygon.VertexInstanceIDs.Reserve(3);
-	Algo::Copy(VertexInstanceIDs, Polygon.VertexInstanceIDs);
-	Polygon.PolygonGroupID = PolygonGroupID;
-	PolygonGroupArray[PolygonGroupID].Polygons.Add(PolygonID);
+	// Make a polygon which will contain this triangle
+	// @todo: make this optional; currently only exists for backwards compatibility
+	FPolygonID PolygonID = PolygonElements->Get().Add();
+	PolygonPolygonGroups[PolygonID] = PolygonGroupID;
+	PolygonGroupToPolygons.AddReferenceToKey(PolygonGroupID, PolygonID);
 
-	Triangle.PolygonID = PolygonID;
-	check(!Polygon.TriangleIDs.Contains(TriangleID));
-	Polygon.TriangleIDs.Add(TriangleID);
+	TrianglePolygons[TriangleID] = PolygonID;
+	PolygonToTriangles.AddReferenceToKey(PolygonID, TriangleID);
 
-	TriangleAttributesSet.Insert(TriangleID);
+	TArrayView<FVertexID> TriVertexIDs = TriangleVertices[TriangleID];
+	TArrayView<FEdgeID> TriEdgeIDs = TriangleEdges[TriangleID];
 
 	for (int32 Index = 0; Index < 3; ++Index)
 	{
-		const FVertexInstanceID VertexInstanceID = Triangle.GetVertexInstanceID(Index);
-		const FVertexInstanceID NextVertexInstanceID = Triangle.GetVertexInstanceID((Index == 2) ? 0 : Index + 1);
+		const FVertexInstanceID VertexInstanceID = TriVertexInstanceIDs[Index];
+		const FVertexInstanceID NextVertexInstanceID = TriVertexInstanceIDs[(Index == 2) ? 0 : Index + 1];
 
-		const FVertexID ThisVertexID = GetVertexInstanceVertex(VertexInstanceID);
-		const FVertexID NextVertexID = GetVertexInstanceVertex(NextVertexInstanceID);
+		const FVertexID ThisVertexID = VertexInstanceVertices[VertexInstanceID];
+		const FVertexID NextVertexID = VertexInstanceVertices[NextVertexInstanceID];
+
+		TriVertexIDs[Index] = ThisVertexID;
+		// VertexToTriangles.AddReferenceToKey(ThisVertexID, TriangleID); // @todo: can add this, but need to consider how degenerates would work
 
 		FEdgeID EdgeID = GetVertexPairEdge(ThisVertexID, NextVertexID);
-		if (EdgeID == FEdgeID::Invalid)
+		if (EdgeID == INDEX_NONE)
 		{
 			EdgeID = CreateEdge(ThisVertexID, NextVertexID);
 			if (OutEdgeIDs)
@@ -418,62 +723,99 @@ void FMeshDescription::CreateTriangle_Internal(const FTriangleID TriangleID, con
 			}
 		}
 
-		check(!VertexInstanceArray[VertexInstanceID].ConnectedTriangles.Contains(TriangleID));
-		VertexInstanceArray[VertexInstanceID].ConnectedTriangles.Add(TriangleID);
+		TriEdgeIDs[Index] = EdgeID;
 
-		check(!EdgeArray[EdgeID].ConnectedTriangles.Contains(TriangleID));
-		EdgeArray[EdgeID].ConnectedTriangles.Add(TriangleID);
+		VertexInstanceToTriangles.AddReferenceToKey(VertexInstanceID, TriangleID);
+		EdgeToTriangles.AddReferenceToKey(EdgeID, TriangleID);
 	}
 }
 
+
+void FMeshDescription::SetTriangleUVIndices(const FTriangleID TriangleID, TArrayView<const FUVID> UVIDs, int32 UVChannel)
+{
+	TArrayView<FUVID> TriUVs = TriangleUVs.Get(TriangleID, UVChannel);
+	TriUVs[0] = UVIDs[0];
+	TriUVs[1] = UVIDs[1];
+	TriUVs[2] = UVIDs[2];
+	UVToTriangles.AddReferenceToKey(UVIDs[0], TriangleID, UVChannel);
+	UVToTriangles.AddReferenceToKey(UVIDs[1], TriangleID, UVChannel);
+	UVToTriangles.AddReferenceToKey(UVIDs[2], TriangleID, UVChannel);
+}
+
+
 template <template <typename...> class TContainer>
 void FMeshDescription::DeleteTriangle_Internal(const FTriangleID TriangleID, TContainer<FEdgeID>* InOutOrphanedEdgesPtr, TContainer<FVertexInstanceID>* InOutOrphanedVertexInstancesPtr, TContainer<FPolygonGroupID>* InOutOrphanedPolygonGroupsPtr)
 {
-	const FMeshTriangle& Triangle = TriangleArray[TriangleID];
-	const FPolygonID PolygonID = Triangle.PolygonID;
-	
-	// Delete this triangle from the polygon
-	verify(PolygonArray[PolygonID].TriangleIDs.RemoveSingle(TriangleID) == 1);
+	const FPolygonGroupID PolygonGroupID = TrianglePolygonGroups[TriangleID];
+	PolygonGroupToTriangles.RemoveReferenceFromKey(PolygonGroupID, TriangleID);
 
-	if (PolygonArray[PolygonID].TriangleIDs.Num() == 0)
+	const FPolygonID PolygonID = TrianglePolygons[TriangleID];
+	PolygonToTriangles.RemoveReferenceFromKey(PolygonID, TriangleID);
+
+	if (PolygonToTriangles.Find(PolygonID).Num() == 0)
 	{
 		// If it was the only triangle in the polygon, delete the polygon too
+		TArrayView<FVertexInstanceID> TriVertexInstanceIDs = TriangleVertexInstances[TriangleID];
+		TArrayView<FEdgeID> TriEdgeIDs = TriangleEdges[TriangleID];
+
 		for (int32 Index = 0; Index < 3; ++Index)
 		{
-			const FVertexInstanceID VertexInstanceID = Triangle.GetVertexInstanceID(Index);
-			const FVertexInstanceID NextVertexInstanceID = Triangle.GetVertexInstanceID((Index == 2) ? 0 : Index + 1);
-
-			const FVertexID VertexID0 = GetVertexInstanceVertex(VertexInstanceID);
-			const FVertexID VertexID1 = GetVertexInstanceVertex(NextVertexInstanceID);
-
-			FEdgeID EdgeID = GetVertexPairEdge(VertexID0, VertexID1);
-			check(EdgeID != FEdgeID::Invalid);
-
-			verify(VertexInstanceArray[VertexInstanceID].ConnectedTriangles.RemoveSingle(TriangleID) == 1);
-			verify(EdgeArray[EdgeID].ConnectedTriangles.RemoveSingle(TriangleID) == 1);
-
-			if (InOutOrphanedVertexInstancesPtr && VertexInstanceArray[VertexInstanceID].ConnectedTriangles.Num() == 0)
+			// Remove vertex instance from the triangle. Here's how we do that safely:
+			// 1) Take a copy of the VertexInstance ID
+			// 2) Remove reference from the VertexInstance to Triangles indexer
+			// 3) Set the triangle VertexInstance to INDEX_NONE
+			// 4) Ask the indexer if that VertexInstance ID has any triangle references, and if not, add it to the list.
+			//
+			// Step 3 is important because Find() below might need to reindex if the key is stale, and this would cause it to
+			// re-add the triangle VertexInstance reference if it were still valid, since the triangle hasn't yet been deleted.
+			//
+			// We could also wait until the triangle were deleted before querying the indexer, but this would involve taking
+			// copies of all the referenced element IDs first.
+
+			const FVertexInstanceID VertexInstanceID = TriVertexInstanceIDs[Index];
+			
+			VertexInstanceToTriangles.RemoveReferenceFromKey(VertexInstanceID, TriangleID);
+			TriVertexInstanceIDs[Index] = INDEX_NONE;
+			
+			if (InOutOrphanedVertexInstancesPtr && VertexInstanceToTriangles.Find(VertexInstanceID).Num() == 0)
 			{
 				AddUnique(*InOutOrphanedVertexInstancesPtr, VertexInstanceID);
 			}
 
-			if (InOutOrphanedEdgesPtr && EdgeArray[EdgeID].ConnectedTriangles.Num() == 0)
+			// Remove edge from the triangle. Same rules as above apply!
+			
+			const FEdgeID EdgeID = TriEdgeIDs[Index];
+			
+			EdgeToTriangles.RemoveReferenceFromKey(EdgeID, TriangleID);
+			TriEdgeIDs[Index] = INDEX_NONE;
+			
+			if (InOutOrphanedEdgesPtr && EdgeToTriangles.Find(EdgeID).Num() == 0)
 			{
 				AddUnique(*InOutOrphanedEdgesPtr, EdgeID);
 			}
+
+			// Remove UVs from the triangle if there are any
+
+			for (int32 UVIndex = 0; UVIndex < TriangleUVs.GetNumChannels(); UVIndex++)
+			{
+				TArrayView<const FUVID> TriUVs = TriangleUVs.Get(TriangleID, UVIndex);
+				FUVID UVID = TriUVs[Index];
+				UVToTriangles.RemoveReferenceFromKey(UVID, TriangleID, UVIndex);
+			}
 		}
 
 		// Remove the polygon
-		const FPolygonGroupID PolygonGroupID = PolygonArray[PolygonID].PolygonGroupID;
-		verify(PolygonGroupArray[PolygonGroupID].Polygons.RemoveSingle(PolygonID) == 1);
+		check(PolygonPolygonGroups[PolygonID] == PolygonGroupID);
+		PolygonGroupToPolygons.RemoveReferenceFromKey(PolygonGroupID, PolygonID);
 
-		if (InOutOrphanedPolygonGroupsPtr && PolygonGroupArray[PolygonGroupID].Polygons.Num() == 0)
+		PolygonElements->Get().Remove(PolygonID);
+		PolygonToTriangles.RemoveKey(PolygonID);
+
+		// Check orphaned polygon groups after the polygon has been removed so that reindexing won't find it again.
+		if (InOutOrphanedPolygonGroupsPtr && PolygonGroupToPolygons.Find(PolygonGroupID).Num() == 0)
 		{
 			AddUnique(*InOutOrphanedPolygonGroupsPtr, PolygonGroupID);
 		}
-
-		PolygonArray.Remove(PolygonID);
-		PolygonAttributesSet.Remove(PolygonID);
 	}
 	else
 	{
@@ -483,8 +825,7 @@ void FMeshDescription::DeleteTriangle_Internal(const FTriangleID TriangleID, TCo
 		check(false);
 	}
 
-	TriangleArray.Remove(TriangleID);
-	TriangleAttributesSet.Remove(TriangleID);
+	TriangleElements->Get().Remove(TriangleID);
 }
 
 void FMeshDescription::DeleteTriangle(const FTriangleID TriangleID, TArray<FEdgeID>* InOutOrphanedEdgesPtr, TArray<FVertexInstanceID>* InOutOrphanedVertexInstancesPtr, TArray<FPolygonGroupID>* InOutOrphanedPolygonGroupsPtr)
@@ -528,21 +869,21 @@ void FMeshDescription::CreatePolygon_Internal(const FPolygonID PolygonID, const
 		OutEdgeIDs->Empty();
 	}
 
-	FMeshPolygon& Polygon = PolygonArray[PolygonID];
-	const int32 NumVertices = VertexInstanceIDs.Num();
-	Polygon.VertexInstanceIDs.SetNumUninitialized(NumVertices);
+	check(PolygonGroupID != INDEX_NONE);
+	PolygonPolygonGroups[PolygonID] = PolygonGroupID;
+	PolygonGroupToPolygons.AddReferenceToKey(PolygonGroupID, PolygonID);
 
+	const int32 NumVertices = VertexInstanceIDs.Num();
 	for (int32 Index = 0; Index < NumVertices; ++Index)
 	{
 		const FVertexInstanceID ThisVertexInstanceID = VertexInstanceIDs[Index];
 		const FVertexInstanceID NextVertexInstanceID = VertexInstanceIDs[(Index + 1 == NumVertices) ? 0 : Index + 1];
-		const FVertexID ThisVertexID = GetVertexInstanceVertex(ThisVertexInstanceID);
-		const FVertexID NextVertexID = GetVertexInstanceVertex(NextVertexInstanceID);
-
-		Polygon.VertexInstanceIDs[Index] = ThisVertexInstanceID;
+		const FVertexID ThisVertexID = VertexInstanceVertices[ThisVertexInstanceID];
+		const FVertexID NextVertexID = VertexInstanceVertices[NextVertexInstanceID];
 
+		// Create any missing perimeter edges here (not internal edges; they will be created when the polygon triangles are created)
 		FEdgeID EdgeID = GetVertexPairEdge(ThisVertexID, NextVertexID);
-		if (EdgeID == FEdgeID::Invalid)
+		if (EdgeID == INDEX_NONE)
 		{
 			EdgeID = CreateEdge(ThisVertexID, NextVertexID);
 			if (OutEdgeIDs)
@@ -552,81 +893,92 @@ void FMeshDescription::CreatePolygon_Internal(const FPolygonID PolygonID, const
 		}
 	}
 
-	check(PolygonGroupID != FPolygonGroupID::Invalid);
-	Polygon.PolygonGroupID = PolygonGroupID;
-	PolygonGroupArray[PolygonGroupID].Polygons.Add(PolygonID);
-
-	check(Polygon.TriangleIDs.Num() == 0);
-	ComputePolygonTriangulation(PolygonID);
-
-	PolygonAttributesSet.Insert(PolygonID);
+	CreatePolygonTriangles(PolygonID, VertexInstanceIDs);
 }
 
 template <template <typename...> class TContainer>
 void FMeshDescription::DeletePolygon_Internal(const FPolygonID PolygonID, TContainer<FEdgeID>* InOutOrphanedEdgesPtr, TContainer<FVertexInstanceID>* InOutOrphanedVertexInstancesPtr, TContainer<FPolygonGroupID>* InOutOrphanedPolygonGroupsPtr)
 {
-	FMeshPolygon& Polygon = PolygonArray[PolygonID];
+	const FPolygonGroupID PolygonGroupID = PolygonPolygonGroups[PolygonID];
 
 	// Remove constituent triangles
-	for (const FTriangleID TriangleID : Polygon.TriangleIDs)
+	TArrayView<const FTriangleID> TriangleIDs = PolygonToTriangles.Find<FTriangleID>(PolygonID);
+
+	TArray<FEdgeID, TInlineAllocator<8>> InternalEdgesToRemove;
+	InternalEdgesToRemove.Reserve(TriangleIDs.Num() - 1);
+
+	// Disconnect edges and vertex instances
+	for (const FTriangleID TriangleID : TriangleIDs)
 	{
-		const FMeshTriangle& Triangle = TriangleArray[TriangleID];
+		TArrayView<FVertexInstanceID> TriVertexInstanceIDs = TriangleVertexInstances[TriangleID];
+		TArrayView<FEdgeID> TriEdgeIDs = TriangleEdges[TriangleID];
 
 		for (int32 Index = 0; Index < 3; ++Index)
 		{
-			const FVertexInstanceID ThisVertexInstanceID = Triangle.GetVertexInstanceID(Index);
-			const FVertexInstanceID NextVertexInstanceID = Triangle.GetVertexInstanceID((Index == 2) ? 0 : Index + 1);
-			const FVertexID ThisVertexID = GetVertexInstanceVertex(ThisVertexInstanceID);
-			const FVertexID NextVertexID = GetVertexInstanceVertex(NextVertexInstanceID);
-			const FEdgeID EdgeID = GetVertexPairEdge(ThisVertexID, NextVertexID);
-
-			// If a valid edge isn't found, we deem this to be because it's an internal edge which was already removed
-			// in a previous iteration through the triangle array.
-			if (EdgeID != FEdgeID::Invalid)
+			const FEdgeID EdgeID = TriEdgeIDs[Index];
+			if (IsEdgeInternal(EdgeID))
 			{
-				if (IsEdgeInternal(EdgeID))
-				{
-					// Remove internal edges
-					for (const FVertexID EdgeVertexID : EdgeArray[EdgeID].VertexIDs)
-					{
-						verify(VertexArray[EdgeVertexID].ConnectedEdgeIDs.RemoveSingle(EdgeID) == 1);
-					}
-					EdgeArray.Remove(EdgeID);
-					EdgeAttributesSet.Remove(EdgeID);
-				}
-				else
-				{
-					verify(EdgeArray[EdgeID].ConnectedTriangles.RemoveSingle(TriangleID) == 1);
+				// Remove internal edges - but not yet.
+				// For now just make a note of them.
+				InternalEdgesToRemove.AddUnique(EdgeID);
+			}
+			else
+			{
+				EdgeToTriangles.RemoveReferenceFromKey(EdgeID, TriangleID);
 
-					if (InOutOrphanedEdgesPtr && EdgeArray[EdgeID].ConnectedTriangles.Num() == 0)
-					{
-						AddUnique(*InOutOrphanedEdgesPtr, EdgeID);
-					}
+				// Set the reference to INDEX_NONE here, so that it won't be picked up in the Find() below
+				// if it needs to rebuild the key.
+				TriEdgeIDs[Index] = INDEX_NONE;
+
+				if (InOutOrphanedEdgesPtr && EdgeToTriangles.Find(EdgeID).Num() == 0)
+				{
+					AddUnique(*InOutOrphanedEdgesPtr, EdgeID);
 				}
 			}
 
-			verify(VertexInstanceArray[ThisVertexInstanceID].ConnectedTriangles.RemoveSingle(TriangleID) == 1);
+			const FVertexInstanceID VertexInstanceID = TriVertexInstanceIDs[Index];
+			VertexInstanceToTriangles.RemoveReferenceFromKey(VertexInstanceID, TriangleID);
+
+			// Set the reference to INDEX_NONE here, so that it won't be picked up in the Find() below
+			// if it needs to rebuild the key.
+			TriVertexInstanceIDs[Index] = INDEX_NONE;
 
-			if (InOutOrphanedVertexInstancesPtr && VertexInstanceArray[ThisVertexInstanceID].ConnectedTriangles.Num() == 0)
+			if (InOutOrphanedVertexInstancesPtr && VertexInstanceToTriangles.Find(VertexInstanceID).Num() == 0)
 			{
-				AddUnique(*InOutOrphanedVertexInstancesPtr, ThisVertexInstanceID);
+				AddUnique(*InOutOrphanedVertexInstancesPtr, VertexInstanceID);
 			}
 		}
-
-		TriangleArray.Remove(TriangleID);
-		TriangleAttributesSet.Remove(TriangleID);
 	}
 
-	FMeshPolygonGroup& PolygonGroup = PolygonGroupArray[Polygon.PolygonGroupID];
-	verify(PolygonGroup.Polygons.RemoveSingle(PolygonID) == 1);
+	// Remove triangles
+	for (const FTriangleID TriangleID : TriangleIDs)
+	{
+		check(TrianglePolygonGroups[TriangleID] == PolygonGroupID);
+		PolygonGroupToTriangles.RemoveReferenceFromKey(PolygonGroupID, TriangleID);
+		TriangleElements->Get().Remove(TriangleID);
+	}
 
-	if (InOutOrphanedPolygonGroupsPtr && PolygonGroup.Polygons.Num() == 0)
+	// Remove internal edges
+	for (const FEdgeID EdgeID : InternalEdgesToRemove)
 	{
-		AddUnique(*InOutOrphanedPolygonGroupsPtr, Polygon.PolygonGroupID);
+		for (const FVertexID EdgeVertexID : EdgeVertices[EdgeID])
+		{
+			VertexToEdges.RemoveReferenceFromKey(EdgeVertexID, EdgeID);
+		}
+
+		EdgeElements->Get().Remove(EdgeID);
+		EdgeToTriangles.RemoveKey(EdgeID);
 	}
 
-	PolygonArray.Remove(PolygonID);
-	PolygonAttributesSet.Remove(PolygonID);
+	PolygonGroupToPolygons.RemoveReferenceFromKey(PolygonGroupID, PolygonID);
+
+	PolygonElements->Get().Remove(PolygonID);
+	PolygonToTriangles.RemoveKey(PolygonID);
+
+	if (InOutOrphanedPolygonGroupsPtr && PolygonGroupToPolygons.Find(PolygonGroupID).Num() == 0)
+	{
+		AddUnique(*InOutOrphanedPolygonGroupsPtr, PolygonGroupID);
+	}
 }
 
 void FMeshDescription::DeletePolygon(const FPolygonID PolygonID, TArray<FEdgeID>* InOutOrphanedEdgesPtr, TArray<FVertexInstanceID>* InOutOrphanedVertexInstancesPtr, TArray<FPolygonGroupID>* InOutOrphanedPolygonGroupsPtr)
@@ -665,9 +1017,11 @@ void FMeshDescription::DeletePolygons(const TArray<FPolygonID>& Polygons)
 
 bool FMeshDescription::IsVertexOrphaned(const FVertexID VertexID) const
 {
-	for (const FVertexInstanceID VertexInstanceID : VertexArray[VertexID].VertexInstanceIDs)
+	TArrayView<const FVertexInstanceID> VertexInstanceIDs = VertexToVertexInstances.Find<FVertexInstanceID>(VertexID);
+	for (const FVertexInstanceID VertexInstanceID : VertexInstanceIDs)
 	{
-		if (VertexInstanceArray[VertexInstanceID].ConnectedTriangles.Num() > 0)
+		TArrayView<const FTriangleID> TriangleIDs = VertexInstanceToTriangles.Find<FTriangleID>(VertexInstanceID);
+		if (TriangleIDs.Num() > 0)
 		{
 			return false;
 		}
@@ -679,67 +1033,83 @@ bool FMeshDescription::IsVertexOrphaned(const FVertexID VertexID) const
 
 FEdgeID FMeshDescription::GetVertexPairEdge(const FVertexID VertexID0, const FVertexID VertexID1) const
 {
-	for (const FEdgeID VertexConnectedEdgeID : VertexArray[VertexID0].ConnectedEdgeIDs)
+	TArrayView<const FEdgeID> ConnectedEdgeIDs = VertexToEdges.Find<FEdgeID>(VertexID0);
+	for (const FEdgeID VertexConnectedEdgeID : ConnectedEdgeIDs)
 	{
-		const FVertexID EdgeVertexID0 = EdgeArray[VertexConnectedEdgeID].VertexIDs[0];
-		const FVertexID EdgeVertexID1 = EdgeArray[VertexConnectedEdgeID].VertexIDs[1];
+		const FVertexID EdgeVertexID0 = EdgeVertices[VertexConnectedEdgeID][0];
+		const FVertexID EdgeVertexID1 = EdgeVertices[VertexConnectedEdgeID][1];
 		if ((EdgeVertexID0 == VertexID0 && EdgeVertexID1 == VertexID1) || (EdgeVertexID0 == VertexID1 && EdgeVertexID1 == VertexID0))
 		{
 			return VertexConnectedEdgeID;
 		}
 	}
 
-	return FEdgeID::Invalid;
+	return INDEX_NONE;
 }
 
 
 FEdgeID FMeshDescription::GetVertexInstancePairEdge(const FVertexInstanceID VertexInstanceID0, const FVertexInstanceID VertexInstanceID1) const
 {
-	const FVertexID VertexID0 = VertexInstanceArray[VertexInstanceID0].VertexID;
-	const FVertexID VertexID1 = VertexInstanceArray[VertexInstanceID1].VertexID;
-	for (const FEdgeID VertexConnectedEdgeID : VertexArray[VertexID0].ConnectedEdgeIDs)
+	const FVertexID VertexID0 = VertexInstanceVertices[VertexInstanceID0];
+	const FVertexID VertexID1 = VertexInstanceVertices[VertexInstanceID1];
+	TArrayView<const FEdgeID> ConnectedEdgeIDs = VertexToEdges.Find<FEdgeID>(VertexID0);
+	for (const FEdgeID VertexConnectedEdgeID : ConnectedEdgeIDs)
 	{
-		const FVertexID EdgeVertexID0 = EdgeArray[VertexConnectedEdgeID].VertexIDs[0];
-		const FVertexID EdgeVertexID1 = EdgeArray[VertexConnectedEdgeID].VertexIDs[1];
+		const FVertexID EdgeVertexID0 = EdgeVertices[VertexConnectedEdgeID][0];
+		const FVertexID EdgeVertexID1 = EdgeVertices[VertexConnectedEdgeID][1];
 		if ((EdgeVertexID0 == VertexID0 && EdgeVertexID1 == VertexID1) || (EdgeVertexID0 == VertexID1 && EdgeVertexID1 == VertexID0))
 		{
 			return VertexConnectedEdgeID;
 		}
 	}
 
-	return FEdgeID::Invalid;
+	return INDEX_NONE;
 }
 
 
 void FMeshDescription::SetPolygonVertexInstance(const FPolygonID PolygonID, const int32 PerimeterIndex, const FVertexInstanceID VertexInstanceID)
 {
-	FMeshPolygon& Polygon = PolygonArray[PolygonID];
-	check(PerimeterIndex >= 0 && PerimeterIndex < Polygon.VertexInstanceIDs.Num());
-
-	// Disconnect old vertex instance from polygon, and connect new one
-	const FVertexInstanceID OldVertexInstanceID = Polygon.VertexInstanceIDs[PerimeterIndex];
-
-	Polygon.VertexInstanceIDs[PerimeterIndex] = VertexInstanceID;
-
-	// Fix up triangle list
-	for (const FTriangleID TriangleID : Polygon.TriangleIDs)
+	TArrayView<const FTriangleID> Tris = PolygonToTriangles.Find<FTriangleID>(PolygonID);
+	if (Tris.Num() == 1)
+	{
+		check(PerimeterIndex < 3);
+		FVertexInstanceID OldVertexInstanceID = TriangleVertexInstances[Tris[0]][PerimeterIndex];
+		check(VertexInstanceVertices[OldVertexInstanceID] == VertexInstanceVertices[VertexInstanceID]);
+		VertexInstanceToTriangles.RemoveReferenceFromKey(OldVertexInstanceID, Tris[0]);
+		TriangleVertexInstances[Tris[0]][PerimeterIndex] = VertexInstanceID;
+		VertexInstanceToTriangles.AddReferenceToKey(VertexInstanceID, Tris[0]);
+	}
+	else
 	{
-		FMeshTriangle& Triangle = TriangleArray[TriangleID];
-		for (int32 VertexIndex = 0; VertexIndex < 3; ++VertexIndex)
+		TArray<FVertexInstanceID, TInlineAllocator<8>> PolygonVertexInstances = GetPolygonVertexInstances<TInlineAllocator<8>>(PolygonID);
+		check(PerimeterIndex < PolygonVertexInstances.Num());
+		FVertexInstanceID OldVertexInstanceID = PolygonVertexInstances[PerimeterIndex];
+		check(VertexInstanceVertices[OldVertexInstanceID] == VertexInstanceVertices[VertexInstanceID]);
+
+		for (const FTriangleID Tri : Tris)
 		{
-			if (Triangle.GetVertexInstanceID(VertexIndex) == OldVertexInstanceID)
+			for (FVertexInstanceID& VertexInstance : TriangleVertexInstances[Tri])
 			{
-				verify(VertexInstanceArray[OldVertexInstanceID].ConnectedTriangles.RemoveSingle(TriangleID) == 1);
-				check(!VertexInstanceArray[VertexInstanceID].ConnectedTriangles.Contains(TriangleID));
-				VertexInstanceArray[VertexInstanceID].ConnectedTriangles.Add(TriangleID);
-				Triangle.SetVertexInstanceID(VertexIndex, VertexInstanceID);
+				if (VertexInstance == OldVertexInstanceID)
+				{
+					VertexInstanceToTriangles.RemoveReferenceFromKey(OldVertexInstanceID, Tri);
+					VertexInstance = VertexInstanceID;
+					VertexInstanceToTriangles.AddReferenceToKey(VertexInstanceID, Tri);
+				}
 			}
 		}
 	}
 }
 
 
-FPlane FMeshDescription::ComputePolygonPlane( const FPolygonID PolygonID ) const
+void FMeshDescription::SetPolygonVertexInstances(const FPolygonID PolygonID, TArrayView<const FVertexInstanceID> VertexInstanceIDs)
+{
+	RemovePolygonTriangles(PolygonID);
+	CreatePolygonTriangles(PolygonID, VertexInstanceIDs);
+}
+
+
+FPlane FMeshDescription::ComputePolygonPlane(TArrayView<const FVertexInstanceID> PerimeterVertexInstanceIDs) const
 {
 	// NOTE: This polygon plane computation code is partially based on the implementation of "Newell's method" from Real-Time
 	//       Collision Detection by Christer Ericson, published by Morgan Kaufmann Publishers, (c) 2005 Elsevier Inc
@@ -752,42 +1122,36 @@ FPlane FMeshDescription::ComputePolygonPlane( const FPolygonID PolygonID ) const
 	FVector Centroid = FVector::ZeroVector;
 	FVector Normal = FVector::ZeroVector;
 
-	TArray<FVertexID> PerimeterVertexIDs;
-	GetPolygonVertices( PolygonID, /* Out */ PerimeterVertexIDs );
-
-	// @todo Maybe this shouldn't be in FMeshDescription but in a utility class, as it references a specific attribute name
-	TVertexAttributesConstRef<FVector> VertexPositions = VertexAttributes().GetAttributesRef<FVector>( MeshAttribute::Vertex::Position );
-
 	// Use 'Newell's Method' to compute a robust 'best fit' plane from the vertices of this polygon
-	for ( int32 VertexNumberI = PerimeterVertexIDs.Num() - 1, VertexNumberJ = 0; VertexNumberJ < PerimeterVertexIDs.Num(); VertexNumberI = VertexNumberJ, VertexNumberJ++ )
+	for (int32 VertexNumberI = PerimeterVertexInstanceIDs.Num() - 1, VertexNumberJ = 0; VertexNumberJ < PerimeterVertexInstanceIDs.Num(); VertexNumberI = VertexNumberJ, VertexNumberJ++)
 	{
-		const FVertexID VertexIDI = PerimeterVertexIDs[ VertexNumberI ];
-		const FVector PositionI = VertexPositions[ VertexIDI ];
+		const FVertexID VertexIDI = PerimeterVertexInstanceIDs[VertexNumberI];
+		const FVector PositionI = VertexPositions[VertexInstanceVertices[VertexIDI]];
 
-		const FVertexID VertexIDJ = PerimeterVertexIDs[ VertexNumberJ ];
-		const FVector PositionJ = VertexPositions[ VertexIDJ ];
+		const FVertexID VertexIDJ = PerimeterVertexInstanceIDs[VertexNumberJ];
+		const FVector PositionJ = VertexPositions[VertexInstanceVertices[VertexIDJ]];
 
 		Centroid += PositionJ;
 
-		Normal.X += ( PositionJ.Y - PositionI.Y ) * ( PositionI.Z + PositionJ.Z );
-		Normal.Y += ( PositionJ.Z - PositionI.Z ) * ( PositionI.X + PositionJ.X );
-		Normal.Z += ( PositionJ.X - PositionI.X ) * ( PositionI.Y + PositionJ.Y );
+		Normal.X += (PositionJ.Y - PositionI.Y) * (PositionI.Z + PositionJ.Z);
+		Normal.Y += (PositionJ.Z - PositionI.Z) * (PositionI.X + PositionJ.X);
+		Normal.Z += (PositionJ.X - PositionI.X) * (PositionI.Y + PositionJ.Y);
 	}
 
 	Normal.Normalize();
 
 	// Construct a plane from the normal and centroid
-	return FPlane( Normal, FVector::DotProduct( Centroid, Normal ) / ( float )PerimeterVertexIDs.Num() );
+	return FPlane(Normal, FVector::DotProduct(Centroid, Normal) / (float)PerimeterVertexInstanceIDs.Num());
 }
 
 
-FVector FMeshDescription::ComputePolygonNormal( const FPolygonID PolygonID ) const
+FVector FMeshDescription::ComputePolygonNormal(TArrayView<const FVertexInstanceID> PerimeterVertexInstanceIDs) const
 {
 	// @todo mesheditor: Polygon normals are now computed and cached when changes are made to a polygon.
 	// In theory, we can just return that cached value, but we need to check that there is nothing which relies on the value being correct before
 	// the cache is updated at the end of a modification.
-	const FPlane PolygonPlane = ComputePolygonPlane( PolygonID );
-	const FVector PolygonNormal( PolygonPlane.X, PolygonPlane.Y, PolygonPlane.Z );
+	const FPlane PolygonPlane = ComputePolygonPlane(PerimeterVertexInstanceIDs);
+	const FVector PolygonNormal(PolygonPlane.X, PolygonPlane.Y, PolygonPlane.Z);
 	return PolygonNormal;
 }
 
@@ -823,121 +1187,341 @@ static bool PointInTriangle(const FVector& A, const FVector& B, const FVector& C
 }
 
 
-void FMeshDescription::ComputePolygonTriangulation(const FPolygonID PolygonID)
+void FMeshDescription::FindPolygonPerimeter(const FPolygonID PolygonID, TArrayView<FEdgeID> Edges) const
 {
-	// NOTE: This polygon triangulation code is partially based on the ear cutting algorithm described on
-	//       page 497 of the book "Real-time Collision Detection", published in 2005.
+	TArrayView<const FTriangleID> PolygonTris = PolygonToTriangles.Find<FTriangleID>(PolygonID);
+	check(Edges.Num() == PolygonTris.Num() + 2);
 
-	FMeshPolygon& Polygon = PolygonArray[PolygonID];
-	const TArray<FVertexInstanceID>& PolygonVertexInstanceIDs = Polygon.VertexInstanceIDs;
-
-	// Polygon must have at least three vertices/edges
-	const int32 PolygonVertexCount = PolygonVertexInstanceIDs.Num();
-	check(PolygonVertexCount >= 3);
-
-	// If polygon was already triangulated, and only has three vertices, no need to do anything here
-	if (Polygon.TriangleIDs.Num() == 1 && PolygonVertexCount == 3)
+	// Optimization for simple triangle
+	if (PolygonTris.Num() == 1)
 	{
+		for (int I = 0; I < 3; I++)
+		{
+			Edges[I] = TriangleEdges[PolygonTris[0]][I];
+		}
 		return;
 	}
 
-	// Remove currently configured triangles
-	for (const FTriangleID TriangleID : Polygon.TriangleIDs)
+	// Determine perimeter for arbitrary n-gon.
+	// @todo: This process can undoubtedly be optimized
+
+	// Build a set of all the perimeter edges
+	TArray<FEdgeID, TInlineAllocator<8>> PerimeterEdges;
+	TArray<int32, TInlineAllocator<8>> TriIndices;
+
+	for (int32 TriIndex = 0; TriIndex < PolygonTris.Num(); TriIndex++)
 	{
-		// Disconnect triangles from vertex instances
-		for (const FVertexInstanceID VertexInstanceID : GetTriangleVertexInstances(TriangleID))
+		FTriangleID PolygonTri = PolygonTris[TriIndex];
+
+		TArrayView<const FEdgeID> TriEdges = TriangleEdges[PolygonTri];
+		for (FEdgeID EdgeID : TriEdges)
 		{
-			verify(VertexInstanceArray[VertexInstanceID].ConnectedTriangles.RemoveSingle(TriangleID) == 1);
+			int32 EdgeIndex = PerimeterEdges.Find(EdgeID);
+			if (EdgeIndex != INDEX_NONE)
+			{
+				// If adding an edge which already exists, it must be an internal edge, so remove it again.
+				PerimeterEdges.RemoveAtSwap(EdgeIndex, 1, false);
+				TriIndices.RemoveAtSwap(EdgeIndex, 1, false);
+			}
+			else
+			{
+				PerimeterEdges.Add(EdgeID);
+				TriIndices.Add(TriIndex);
+			}
 		}
+	}
+	check(PerimeterEdges.Num() == PolygonTris.Num() + 2);
+	check(PerimeterEdges.Num() == TriIndices.Num());
 
-		// Disconnect triangles from perimeter edges, and delete internal edges
-		for (const FEdgeID EdgeID : GetTriangleEdges(TriangleID))
+	// Reorder edges to be adjacent by ensuring there's a mutual vertex in consecutive edges
+	for (int32 EdgeIndex = 0; EdgeIndex < PerimeterEdges.Num() - 2; EdgeIndex++)
+	{
+		TArrayView<const FVertexID> EdgeVertexIDs = EdgeVertices[PerimeterEdges[EdgeIndex]];
+		for (int32 NextEdgeIndex = EdgeIndex + 1; NextEdgeIndex < PerimeterEdges.Num(); NextEdgeIndex++)
 		{
-			if (EdgeID != FEdgeID::Invalid)
+			TArrayView<const FVertexID> NextEdgeVertexIDs = EdgeVertices[PerimeterEdges[NextEdgeIndex]];
+			if (EdgeVertexIDs[0] == NextEdgeVertexIDs[0] || EdgeVertexIDs[0] == NextEdgeVertexIDs[1] ||
+				EdgeVertexIDs[1] == NextEdgeVertexIDs[0] || EdgeVertexIDs[1] == NextEdgeVertexIDs[1])
 			{
-				// The edge may be invalid if it was an internal edge which was deleted in a previous iteration through the triangles.
-				// So only do something with valid edges
-				if (IsEdgeInternal(EdgeID))
+				if (NextEdgeIndex > EdgeIndex + 1)
 				{
-					// Remove internal edges completely (the first time they are seen)
-					for (const FVertexID VertexID : GetEdgeVertices(EdgeID))
-					{
-						// Disconnect edge from vertices
-						verify(VertexArray[VertexID].ConnectedEdgeIDs.RemoveSingle(EdgeID) == 1);
-					}
-
-					EdgeArray.Remove(EdgeID);
-					EdgeAttributesSet.Remove(EdgeID);
+					Swap(PerimeterEdges[EdgeIndex + 1], PerimeterEdges[NextEdgeIndex]);
+					Swap(TriIndices[EdgeIndex + 1], TriIndices[NextEdgeIndex]);
 				}
-				else
+				break;
+			}
+		}
+	}
+	check(EdgeVertices[PerimeterEdges.Last()][0] == EdgeVertices[PerimeterEdges[0]][0] || EdgeVertices[PerimeterEdges.Last()][0] == EdgeVertices[PerimeterEdges[0]][1] ||
+		  EdgeVertices[PerimeterEdges.Last()][1] == EdgeVertices[PerimeterEdges[0]][0] || EdgeVertices[PerimeterEdges.Last()][1] == EdgeVertices[PerimeterEdges[0]][1]);
+
+	// Swap the winding order if incorrect
+	FEdgeID FirstEdge = PerimeterEdges[0];
+	FEdgeID SecondEdge = PerimeterEdges[1];
+
+	// Get the triangle which the first edge lies in
+	FTriangleID FirstTriangle = PolygonTris[TriIndices[0]];
+	int32 TriEdgeIndex = TriangleEdges[FirstTriangle].Find(FirstEdge);
+
+	// Get the end vertex of the edge as used in that triangle
+	FVertexID SecondVertex = TriangleVertices[FirstTriangle][(TriEdgeIndex + 1) % 3];
+
+	// If the second edge doesn't contain that end vertex, we need to reverse the order of the edges we just constructed
+	if (EdgeVertices[SecondEdge][0] != SecondVertex && EdgeVertices[SecondEdge][1] != SecondVertex)
+	{
+		for (int32 I = 0; I < PerimeterEdges.Num() / 2; I++)
+		{
+			Swap(PerimeterEdges[I], PerimeterEdges[PerimeterEdges.Num() - 1 - I]);
+			Swap(TriIndices[I], TriIndices[PerimeterEdges.Num() - 1 - I]);
+		}
+	}
+
+	for (int I = 0; I < PerimeterEdges.Num(); I++)
+	{
+		Edges[I] = PerimeterEdges[I];
+	}
+}
+
+
+void FMeshDescription::FindPolygonPerimeter(TArrayView<const FTriangleID> Triangles, TArrayView<TTuple<int32, int32>> Result) const
+{
+	// This constructs the perimeter indices for a polygon in the correct winding order from its constituent triangles
+	check(Result.Num() == Triangles.Num() + 2);
+
+	// Optimization for simple triangle
+	if (Triangles.Num() == 1)
+	{
+		// Return triangle index 0; edges 0, 1, 2
+		Result[0] = MakeTuple(0, 0);
+		Result[1] = MakeTuple(0, 1);
+		Result[2] = MakeTuple(0, 2);
+		return;
+	}
+
+	// Determine perimeter for arbitrary n-gon.
+	// @todo: This process can undoubtedly be optimized
+
+	// Build a set of all the perimeter edges
+	TArray<FEdgeID, TInlineAllocator<8>> PerimeterEdges;
+	TArray<TTuple<int32, int32>, TInlineAllocator<8>> Indices;
+
+	for (int32 TriIndex = 0; TriIndex < Triangles.Num(); TriIndex++)
+	{
+		FTriangleID PolygonTri = Triangles[TriIndex];
+
+		TArrayView<const FEdgeID> TriEdges = TriangleEdges[PolygonTri];
+		for (int32 EdgeIndex = 0; EdgeIndex < 3; EdgeIndex++)
+		{
+			FEdgeID EdgeID = TriEdges[EdgeIndex];
+			int32 PerimeterIndex = PerimeterEdges.Find(EdgeID);
+			if (PerimeterIndex != INDEX_NONE)
+			{
+				// If adding an edge which already exists, it must be an internal edge, so remove it again.
+				PerimeterEdges.RemoveAtSwap(PerimeterIndex, 1, false);
+				Indices.RemoveAtSwap(PerimeterIndex, 1, false);
+			}
+			else
+			{
+				PerimeterEdges.Add(EdgeID);
+				Indices.Add(MakeTuple(TriIndex, EdgeIndex));
+			}
+		}
+	}
+	check(PerimeterEdges.Num() == Triangles.Num() + 2);
+	check(PerimeterEdges.Num() == Indices.Num());
+
+	// Reorder edges to be adjacent by ensuring there's a mutual vertex in consecutive edges
+	for (int32 EdgeIndex = 0; EdgeIndex < PerimeterEdges.Num() - 2; EdgeIndex++)
+	{
+		TArrayView<const FVertexID> EdgeVertexIDs = EdgeVertices[PerimeterEdges[EdgeIndex]];
+		for (int32 NextEdgeIndex = EdgeIndex + 1; NextEdgeIndex < PerimeterEdges.Num(); NextEdgeIndex++)
+		{
+			TArrayView<const FVertexID> NextEdgeVertexIDs = EdgeVertices[PerimeterEdges[NextEdgeIndex]];
+			if (EdgeVertexIDs[0] == NextEdgeVertexIDs[0] || EdgeVertexIDs[0] == NextEdgeVertexIDs[1] ||
+				EdgeVertexIDs[1] == NextEdgeVertexIDs[0] || EdgeVertexIDs[1] == NextEdgeVertexIDs[1])
+			{
+				if (NextEdgeIndex > EdgeIndex + 1)
 				{
-					// Don't remove perimeter edge, but disconnect this triangle from it
-					verify(EdgeArray[EdgeID].ConnectedTriangles.RemoveSingle(TriangleID) == 1);
+					Swap(PerimeterEdges[EdgeIndex + 1], PerimeterEdges[NextEdgeIndex]);
+					Swap(Indices[EdgeIndex + 1], Indices[NextEdgeIndex]);
 				}
+				break;
 			}
 		}
+	}
+	check(EdgeVertices[PerimeterEdges.Last()][0] == EdgeVertices[PerimeterEdges[0]][0] || EdgeVertices[PerimeterEdges.Last()][0] == EdgeVertices[PerimeterEdges[0]][1] ||
+		  EdgeVertices[PerimeterEdges.Last()][1] == EdgeVertices[PerimeterEdges[0]][0] || EdgeVertices[PerimeterEdges.Last()][1] == EdgeVertices[PerimeterEdges[0]][1]);
+
+	// Swap the winding order if incorrect
+	FEdgeID FirstEdge = PerimeterEdges[0];
+	FEdgeID SecondEdge = PerimeterEdges[1];
+
+	// Get the triangle which the first edge lies in
+	FTriangleID FirstTriangle = Triangles[Indices[0].Get<0>()];
+	int32 TriEdgeIndex = TriangleEdges[FirstTriangle].Find(FirstEdge);
+
+	// Get the end vertex of the edge as used in that triangle
+	FVertexID SecondVertex = TriangleVertices[FirstTriangle][(TriEdgeIndex + 1) % 3];
+
+	// If the second edge doesn't contain that end vertex, we need to reverse the order of the edges we just constructed
+	if (EdgeVertices[SecondEdge][0] != SecondVertex && EdgeVertices[SecondEdge][1] != SecondVertex)
+	{
+		for (int32 I = 0; I < Indices.Num() / 2; I++)
+		{
+			Swap(Indices[I], Indices[PerimeterEdges.Num() - 1 - I]);
+		}
+	}
+
+	for (int I = 0; I < Indices.Num(); I++)
+	{
+		Result[I] = Indices[I];
+	}
+}
+
+
+void FMeshDescription::ComputePolygonTriangulation(const FPolygonID PolygonID)
+{
+	TArrayView<const FTriangleID> TriangleIDs = PolygonToTriangles.Find<FTriangleID>(PolygonID);
+
+	// Not valid to call this on an untriangulated polygon - an untriangulated polygon is no longer valid at all.
+	check(TriangleIDs.Num() > 0);
 
-		TriangleArray.Remove(TriangleID);
-		TriangleAttributesSet.Remove(TriangleID);
+	// If polygon was already triangulated, and only has three vertices, no need to do anything here
+	if (TriangleIDs.Num() == 1)
+	{
+		return;
 	}
 
-	Polygon.TriangleIDs.Reset();
+	// Get the current perimeter vertex instances
+	TArray<FVertexInstanceID, TInlineAllocator<8>> PolygonVertexInstanceIDs = GetPolygonVertexInstances<TInlineAllocator<8>>(PolygonID);
+
+	// Remove existing triangles
+	RemovePolygonTriangles(PolygonID);
+
+	// Perform the triangulation
+	CreatePolygonTriangles(PolygonID, PolygonVertexInstanceIDs);
+}
+
+
+void FMeshDescription::RemovePolygonTriangles(const FPolygonID PolygonID)
+{
+	TArrayView<const FTriangleID> TriangleIDs = PolygonToTriangles.Find<FTriangleID>(PolygonID);
+	const FPolygonGroupID PolygonGroupID = PolygonPolygonGroups[PolygonID];
+
+	// Remove currently configured triangles
+	TArray<FEdgeID> InternalEdgesToRemove;
+	InternalEdgesToRemove.Reserve(TriangleIDs.Num() - 1);
+
+	for (const FTriangleID TriangleID : TriangleIDs)
+	{
+		TArrayView<FVertexInstanceID> TriVertexInstanceIDs = TriangleVertexInstances[TriangleID];
+		TArrayView<FEdgeID> TriEdgeIDs = TriangleEdges[TriangleID];
+
+		for (int32 Index = 0; Index < 3; ++Index)
+		{
+			VertexInstanceToTriangles.RemoveReferenceFromKey(TriVertexInstanceIDs[Index], TriangleID);
+
+			const FEdgeID EdgeID = TriEdgeIDs[Index];
+			if (IsEdgeInternal(EdgeID))
+			{
+				InternalEdgesToRemove.Add(EdgeID);
+			}
+
+			EdgeToTriangles.RemoveReferenceFromKey(EdgeID, TriangleID);
+		}
+
+		PolygonGroupToTriangles.RemoveReferenceFromKey(PolygonGroupID, TriangleID);
+		TriangleElements->Get().Remove(TriangleID);
+	}
+
+	PolygonToTriangles.RemoveKey(PolygonID);
+
+	// Remove internal edges
+	for (const FEdgeID EdgeID : InternalEdgesToRemove)
+	{
+		for (const FVertexID EdgeVertexID : EdgeVertices[EdgeID])
+		{
+			VertexToEdges.RemoveReferenceFromKey(EdgeVertexID, EdgeID);
+		}
+
+		EdgeElements->Get().Remove(EdgeID);
+		EdgeToTriangles.RemoveKey(EdgeID);
+	}
+}
+
+void FMeshDescription::CreatePolygonTriangles(const FPolygonID PolygonID, TArrayView<const FVertexInstanceID> VertexInstanceIDs)
+{
+	FPolygonGroupID PolygonGroupID = PolygonPolygonGroups[PolygonID];
 
 	// If perimeter only has 3 vertices, just add a single triangle and return
-	if (PolygonVertexCount == 3)
+	if (VertexInstanceIDs.Num() == 3)
 	{
-		const FTriangleID TriangleID = TriangleArray.Add();
-		TriangleAttributesSet.Insert(TriangleID);
-		
-		FMeshTriangle& Triangle = TriangleArray[TriangleID];
-		Triangle.PolygonID = PolygonID;
-		Polygon.TriangleIDs.Add(TriangleID);
+		const FTriangleID TriangleID = TriangleElements->Get().Add();
+
+		// Fill out triangle vertex instances
+		TArrayView<FVertexInstanceID> TriVertexInstanceIDs = TriangleVertexInstances[TriangleID];
+		TriVertexInstanceIDs[0] = VertexInstanceIDs[0];
+		TriVertexInstanceIDs[1] = VertexInstanceIDs[1];
+		TriVertexInstanceIDs[2] = VertexInstanceIDs[2];
+
+		// Fill out triangle polygon group
+		TrianglePolygonGroups[TriangleID] = PolygonGroupID;
+		PolygonGroupToTriangles.AddReferenceToKey(PolygonGroupID, TriangleID);
+
+		// Fill out triangle polygon
+		TrianglePolygons[TriangleID] = PolygonID;
+		PolygonToTriangles.AddReferenceToKey(PolygonID, TriangleID);
+
+		TArrayView<FVertexID> TriVertexIDs = TriangleVertices[TriangleID];
+		TArrayView<FEdgeID> TriEdgeIDs = TriangleEdges[TriangleID];
 
 		for (int32 Index = 0; Index < 3; ++Index)
 		{
-			const FVertexInstanceID ThisVertexInstanceID = PolygonVertexInstanceIDs[Index];
-			const FVertexInstanceID NextVertexInstanceID = PolygonVertexInstanceIDs[(Index == 2) ? 0 : Index + 1];
-			const FVertexID ThisVertexID = GetVertexInstanceVertex(ThisVertexInstanceID);
-			const FVertexID NextVertexID = GetVertexInstanceVertex(NextVertexInstanceID);
-			const FEdgeID EdgeID = GetVertexPairEdge(ThisVertexID, NextVertexID);
-			check(EdgeID != FEdgeID::Invalid);
+			const FVertexInstanceID VertexInstanceID = TriVertexInstanceIDs[Index];
+			const FVertexInstanceID NextVertexInstanceID = TriVertexInstanceIDs[(Index == 2) ? 0 : Index + 1];
+
+			const FVertexID ThisVertexID = VertexInstanceVertices[VertexInstanceID];
+			const FVertexID NextVertexID = VertexInstanceVertices[NextVertexInstanceID];
 
-			Triangle.SetVertexInstanceID(Index, ThisVertexInstanceID);
+			TriVertexIDs[Index] = ThisVertexID;
 
-			check(!EdgeArray[EdgeID].ConnectedTriangles.Contains(TriangleID));
-			EdgeArray[EdgeID].ConnectedTriangles.Add(TriangleID);
+			FEdgeID EdgeID = GetVertexPairEdge(ThisVertexID, NextVertexID);
+			check(EdgeID != INDEX_NONE);
+			TriEdgeIDs[Index] = EdgeID;
 
-			check(!VertexInstanceArray[ThisVertexInstanceID].ConnectedTriangles.Contains(TriangleID));
-			VertexInstanceArray[ThisVertexInstanceID].ConnectedTriangles.Add(TriangleID);
+			VertexInstanceToTriangles.AddReferenceToKey(VertexInstanceID, TriangleID);
+			EdgeToTriangles.AddReferenceToKey(EdgeID, TriangleID);
 		}
 
 		return;
 	}
 
+	// NOTE: This polygon triangulation code is partially based on the ear cutting algorithm described on
+	//       page 497 of the book "Real-time Collision Detection", published in 2005.
+
 	// @todo mesheditor: Perhaps should always attempt to triangulate by splitting polygons along the shortest edge, for better determinism.
 
 	// First figure out the polygon normal.  We need this to determine which triangles are convex, so that
 	// we can figure out which ears to clip
-	const FVector PolygonNormal = ComputePolygonNormal(PolygonID);
+	const FVector PolygonNormal = ComputePolygonNormal(VertexInstanceIDs);
 
 	// Make a simple linked list array of the previous and next vertex numbers, for each vertex number
 	// in the polygon.  This will just save us having to iterate later on.
 	TArray<int32> PrevVertexNumbers;
 	TArray<int32> NextVertexNumbers;
-	TArray<FVector> VertexPositions;
-
+	TArray<FVector> PolyVertexPositions;
+	int32 PolygonVertexCount = VertexInstanceIDs.Num();
 	{
-		const TVertexAttributesRef<FVector> MeshVertexPositions = VertexAttributes().GetAttributesRef<FVector>(MeshAttribute::Vertex::Position);
 		PrevVertexNumbers.SetNumUninitialized(PolygonVertexCount, false);
 		NextVertexNumbers.SetNumUninitialized(PolygonVertexCount, false);
-		VertexPositions.SetNumUninitialized(PolygonVertexCount, false);
+		PolyVertexPositions.SetNumUninitialized(PolygonVertexCount, false);
 
 		for (int32 VertexNumber = 0; VertexNumber < PolygonVertexCount; ++VertexNumber)
 		{
 			PrevVertexNumbers[VertexNumber] = VertexNumber - 1;
 			NextVertexNumbers[VertexNumber] = VertexNumber + 1;
 
-			VertexPositions[VertexNumber] = MeshVertexPositions[GetVertexInstanceVertex(PolygonVertexInstanceIDs[VertexNumber])];
+			PolyVertexPositions[VertexNumber] = VertexPositions[GetVertexInstanceVertex(VertexInstanceIDs[VertexNumber])];
 		}
 		PrevVertexNumbers[0] = PolygonVertexCount - 1;
 		NextVertexNumbers[PolygonVertexCount - 1] = 0;
@@ -954,9 +1538,9 @@ void FMeshDescription::ComputePolygonTriangulation(const FPolygonID PolygonID)
 		// vertices are colinear or other degenerate cases.
 		if (RemainingVertexCount > 3 && EarTestCount < RemainingVertexCount)
 		{
-			const FVector PrevVertexPosition = VertexPositions[PrevVertexNumbers[EarVertexNumber]];
-			const FVector EarVertexPosition = VertexPositions[EarVertexNumber];
-			const FVector NextVertexPosition = VertexPositions[NextVertexNumbers[EarVertexNumber]];
+			const FVector PrevVertexPosition = PolyVertexPositions[PrevVertexNumbers[EarVertexNumber]];
+			const FVector EarVertexPosition = PolyVertexPositions[EarVertexNumber];
+			const FVector NextVertexPosition = PolyVertexPositions[NextVertexNumbers[EarVertexNumber]];
 
 			// Figure out whether the potential ear triangle is facing the same direction as the polygon
 			// itself.  If it's facing the opposite direction, then we're dealing with a concave triangle
@@ -969,7 +1553,7 @@ void FMeshDescription::ComputePolygonTriangulation(const FPolygonID PolygonID)
 				{
 					// Test every other remaining vertex to make sure that it doesn't lie inside our potential ear
 					// triangle.  If we find a vertex that's inside the triangle, then it cannot actually be an ear.
-					const FVector TestVertexPosition = VertexPositions[TestVertexNumber];
+					const FVector TestVertexPosition = PolyVertexPositions[TestVertexNumber];
 					if (PointInTriangle(PrevVertexPosition, EarVertexPosition, NextVertexPosition, TestVertexPosition, SMALL_NUMBER))
 					{
 						bIsEar = false;
@@ -992,37 +1576,45 @@ void FMeshDescription::ComputePolygonTriangulation(const FPolygonID PolygonID)
 			// This will also create any missing internal edges.
 			{
 				// Add a new triangle
-				const FTriangleID TriangleID = TriangleArray.Add();
-				TriangleAttributesSet.Insert(TriangleID);
-
-				// Set its vertex instances and connect it to its parent polygon
-				FMeshTriangle& Triangle = TriangleArray[TriangleID];
-				Triangle.SetVertexInstanceID(0, PolygonVertexInstanceIDs[PrevVertexNumbers[EarVertexNumber]]);
-				Triangle.SetVertexInstanceID(1, PolygonVertexInstanceIDs[EarVertexNumber]);
-				Triangle.SetVertexInstanceID(2, PolygonVertexInstanceIDs[NextVertexNumbers[EarVertexNumber]]);
-				Triangle.PolygonID = PolygonID;
-				check(!Polygon.TriangleIDs.Contains(TriangleID));
-				Polygon.TriangleIDs.Add(TriangleID);
-
-				// Now generate internal edges and connected vertex instances to the new triangle
+				const FTriangleID TriangleID = TriangleElements->Get().Add();
+
+				// Fill out triangle vertex instances
+				TArrayView<FVertexInstanceID> TriVertexInstanceIDs = TriangleVertexInstances[TriangleID];
+				TriVertexInstanceIDs[0] = VertexInstanceIDs[PrevVertexNumbers[EarVertexNumber]];
+				TriVertexInstanceIDs[1] = VertexInstanceIDs[EarVertexNumber];
+				TriVertexInstanceIDs[2] = VertexInstanceIDs[NextVertexNumbers[EarVertexNumber]];
+
+				// Fill out triangle polygon group
+				TrianglePolygonGroups[TriangleID] = PolygonGroupID;
+				PolygonGroupToTriangles.AddReferenceToKey(PolygonGroupID, TriangleID);
+
+				// Fill out triangle polygon
+				TrianglePolygons[TriangleID] = PolygonID;
+				PolygonToTriangles.AddReferenceToKey(PolygonID, TriangleID);
+
+				TArrayView<FVertexID> TriVertexIDs = TriangleVertices[TriangleID];
+				TArrayView<FEdgeID> TriEdgeIDs = TriangleEdges[TriangleID];
+
 				for (int32 Index = 0; Index < 3; ++Index)
 				{
-					const FVertexInstanceID ThisVertexInstanceID = Triangle.GetVertexInstanceID(Index);
-					const FVertexInstanceID NextVertexInstanceID = Triangle.GetVertexInstanceID((Index == 2) ? 0 : Index + 1);
-					const FVertexID ThisVertexID = GetVertexInstanceVertex(ThisVertexInstanceID);
-					const FVertexID NextVertexID = GetVertexInstanceVertex(NextVertexInstanceID);
+					const FVertexInstanceID VertexInstanceID = TriVertexInstanceIDs[Index];
+					const FVertexInstanceID NextVertexInstanceID = TriVertexInstanceIDs[(Index == 2) ? 0 : Index + 1];
+
+					const FVertexID ThisVertexID = VertexInstanceVertices[VertexInstanceID];
+					const FVertexID NextVertexID = VertexInstanceVertices[NextVertexInstanceID];
+
+					TriVertexIDs[Index] = ThisVertexID;
+
 					FEdgeID EdgeID = GetVertexPairEdge(ThisVertexID, NextVertexID);
-					if (EdgeID == FEdgeID::Invalid)
+					if (EdgeID == INDEX_NONE)
 					{
 						// This must be an internal edge (as perimeter edges will already be defined)
 						EdgeID = CreateEdge(ThisVertexID, NextVertexID);
 					}
+					TriEdgeIDs[Index] = EdgeID;
 
-					check(!VertexInstanceArray[ThisVertexInstanceID].ConnectedTriangles.Contains(TriangleID));
-					VertexInstanceArray[ThisVertexInstanceID].ConnectedTriangles.Add(TriangleID);
-
-					check(!EdgeArray[EdgeID].ConnectedTriangles.Contains(TriangleID));
-					EdgeArray[EdgeID].ConnectedTriangles.Add(TriangleID);
+					VertexInstanceToTriangles.AddReferenceToKey(VertexInstanceID, TriangleID);
+					EdgeToTriangles.AddReferenceToKey(EdgeID, TriangleID);
 				}
 			}
 
@@ -1050,8 +1642,6 @@ void FMeshDescription::ComputePolygonTriangulation(const FPolygonID PolygonID)
 			++EarTestCount;
 		}
 	}
-
-	check(Polygon.TriangleIDs.Num() > 0);
 }
 
 
@@ -1062,8 +1652,6 @@ FBoxSphereBounds FMeshDescription::GetBounds() const
 	FBox BoundingBox;
 	BoundingBox.Init();
 
-	TVertexAttributesConstRef<FVector> VertexPositions = VertexAttributes().GetAttributesRef<FVector>(MeshAttribute::Vertex::Position);
-
 	for (const FVertexID VertexID : Vertices().GetElementIDs())
 	{
 		if (!IsVertexOrphaned(VertexID))
@@ -1098,6 +1686,19 @@ void FMeshDescription::TriangulateMesh()
 }
 
 
+void FMeshDescription::SetNumUVChannels(const int32 NumUVChannels)
+{
+	UVElements->SetNumChannels(NumUVChannels);
+	TriangleUVs = TriangleElements->Get().GetAttributes().RegisterIndexAttribute<FUVID[3]>(MeshAttribute::Triangle::UVIndex, NumUVChannels);
+
+	// Ensure every UV element channel has a UVCoordinate attribute
+	for (int32 Index = 0; Index < NumUVChannels; Index++)
+	{
+		UVElements->Get(Index).GetAttributes().RegisterAttribute(MeshAttribute::UV::UVCoordinate, 1, FVector2D::ZeroVector, EMeshAttributeFlags::Lerpable);
+	}
+}
+
+
 namespace MeshAttribute_
 {
 	namespace Vertex
@@ -1140,7 +1741,7 @@ namespace MeshAttribute_
 
 float FMeshDescription::GetPolygonCornerAngleForVertex(const FPolygonID PolygonID, const FVertexID VertexID) const
 {
-	const FMeshPolygon& Polygon = PolygonArray[PolygonID];
+	TArray<FVertexInstanceID> PolygonVertexInstanceIDs = GetPolygonVertexInstances(PolygonID);
 
 	// Lambda function which returns the inner angle at a given index on a polygon contour
 	auto GetContourAngle = [this](const TArray<FVertexInstanceID>& VertexInstanceIDs, const int32 ContourIndex)
@@ -1166,18 +1767,17 @@ float FMeshDescription::GetPolygonCornerAngleForVertex(const FPolygonID PolygonI
 		return FMath::Acos(FVector::DotProduct(Direction1, Direction2));
 	};
 
-	const FVertexInstanceArray& VertexInstancesRef = VertexInstances();
-	auto IsVertexInstancedFromThisVertex = [&VertexInstancesRef, VertexID](const FVertexInstanceID VertexInstanceID)
+	auto IsVertexInstancedFromThisVertex = [this, VertexID](const FVertexInstanceID VertexInstanceID)
 	{
-		return VertexInstancesRef[VertexInstanceID].VertexID == VertexID;
+		return this->GetVertexInstanceVertex(VertexInstanceID) == VertexID;
 	};
 
 	// First look for the vertex instance in the perimeter
-	int32 ContourIndex = Polygon.VertexInstanceIDs.IndexOfByPredicate(IsVertexInstancedFromThisVertex);
+	int32 ContourIndex = PolygonVertexInstanceIDs.IndexOfByPredicate(IsVertexInstancedFromThisVertex);
 	if (ContourIndex != INDEX_NONE)
 	{
 		// Return the internal angle if found
-		return GetContourAngle(Polygon.VertexInstanceIDs, ContourIndex);
+		return GetContourAngle(PolygonVertexInstanceIDs, ContourIndex);
 	}
 
 	// Found nothing; return 0
@@ -1188,8 +1788,6 @@ FBox FMeshDescription::ComputeBoundingBox() const
 {
 	FBox BoundingBox(ForceInit);
 
-	TVertexAttributesConstRef<FVector> VertexPositions = VertexAttributes().GetAttributesRef<FVector>(MeshAttribute::Vertex::Position);
-
 	for (const FVertexID VertexID : Vertices().GetElementIDs())
 	{
 		BoundingBox += VertexPositions[VertexID];
@@ -1199,21 +1797,28 @@ FBox FMeshDescription::ComputeBoundingBox() const
 }
 
 
+void FMeshDescription::ReverseTriangleFacing(const FTriangleID TriangleID)
+{
+	TArrayView<FVertexInstanceID> TriVertexInstances = TriangleVertexInstances[TriangleID];
+	TArrayView<FVertexID> TriVertices = TriangleVertices[TriangleID];
+	TArrayView<FEdgeID> TriEdges = TriangleEdges[TriangleID];
+	Swap(TriVertexInstances[0], TriVertexInstances[1]);
+	Swap(TriVertices[0], TriVertices[1]);
+	Swap(TriEdges[0], TriEdges[1]);
+}
+
+
 void FMeshDescription::ReversePolygonFacing(const FPolygonID PolygonID)
 {
 	// Build a reverse perimeter
-	FMeshPolygon& Polygon = PolygonArray[PolygonID];
-	for (int32 i = 0; i < Polygon.VertexInstanceIDs.Num() / 2; ++i)
+	TArray<FVertexInstanceID> Contour = GetPolygonVertexInstances(PolygonID);
+	for (int32 i = 0; i < Contour.Num() / 2; ++i)
 	{
-		Polygon.VertexInstanceIDs.Swap(i, Polygon.VertexInstanceIDs.Num() - i - 1);
+		Swap(Contour[i], Contour[Contour.Num() - i - 1]);
 	}
 
-	// Update the polygon's triangle vertex instance ids with the reversed ids
-	for (FTriangleID TriangleID : GetPolygonTriangleIDs(PolygonID))
-	{
-		FMeshTriangle& Triangle = TriangleArray[TriangleID];
-		Swap(Triangle.VertexInstanceIDs[0], Triangle.VertexInstanceIDs[2]);
-	}
+	RemovePolygonTriangles(PolygonID);
+	CreatePolygonTriangles(PolygonID, Contour);
 }
 
 
@@ -1234,11 +1839,12 @@ void FMeshDescription::RemapPolygonGroups(const TMap<FPolygonGroupID, FPolygonGr
 	struct FOldPolygonGroupData
 	{
 		FName Name;
+		TArray<FTriangleID> Triangles;
 		TArray<FPolygonID> Polygons;
 	};
 
 	TMap<FPolygonGroupID, FOldPolygonGroupData> OldData;
-	for (const FPolygonGroupID& PolygonGroupID : PolygonGroups().GetElementIDs())
+	for (const FPolygonGroupID PolygonGroupID : PolygonGroups().GetElementIDs())
 	{
 		if (!Remap.Contains(PolygonGroupID) || PolygonGroupID == Remap[PolygonGroupID])
 		{
@@ -1247,11 +1853,13 @@ void FMeshDescription::RemapPolygonGroups(const TMap<FPolygonGroupID, FPolygonGr
 		}
 		FOldPolygonGroupData& PolygonGroupData = OldData.FindOrAdd(PolygonGroupID);
 		PolygonGroupData.Name = PolygonGroupNames[PolygonGroupID];
-		FMeshPolygonGroup& PolygonGroup = PolygonGroupArray[PolygonGroupID];
-		PolygonGroupData.Polygons = PolygonGroup.Polygons;
-		PolygonGroup.Polygons.Empty();
-		DeletePolygonGroup(PolygonGroupID);
+		PolygonGroupData.Triangles = PolygonGroupToTriangles.Find<FTriangleID>(PolygonGroupID);
+		PolygonGroupData.Polygons = PolygonGroupToPolygons.Find<FPolygonID>(PolygonGroupID);
+		PolygonGroupElements->Get().Remove(PolygonGroupID);
+		PolygonGroupToPolygons.RemoveKey(PolygonGroupID);
+		PolygonGroupToTriangles.RemoveKey(PolygonGroupID);
 	}
+
 	for (auto Kvp : OldData)
 	{
 		FPolygonGroupID GroupID = Kvp.Key;
@@ -1260,12 +1868,18 @@ void FMeshDescription::RemapPolygonGroups(const TMap<FPolygonGroupID, FPolygonGr
 		{
 			CreatePolygonGroupWithID(ToGroupID);
 		}
-		TArray<FPolygonID>& Polygons = PolygonGroupArray[ToGroupID].Polygons;
-		Polygons.Append(Kvp.Value.Polygons);
 		PolygonGroupNames[ToGroupID] = Kvp.Value.Name;
-		for (FPolygonID PolygonID : Polygons)
+
+		for (const FTriangleID TriangleID : Kvp.Value.Triangles)
+		{
+			TrianglePolygonGroups[TriangleID] = ToGroupID;
+			PolygonGroupToTriangles.AddReferenceToKey(ToGroupID, TriangleID);
+		}
+
+		for (const FPolygonID PolygonID : Kvp.Value.Polygons)
 		{
-			PolygonArray[PolygonID].PolygonGroupID = ToGroupID;
+			PolygonPolygonGroups[PolygonID] = ToGroupID;
+			PolygonGroupToPolygons.AddReferenceToKey(ToGroupID, PolygonID);
 		}
 	}
 }
diff --git a/Engine/Source/Runtime/MeshDescription/Private/MeshDescriptionBase.cpp b/Engine/Source/Runtime/MeshDescription/Private/MeshDescriptionBase.cpp
--- a/Engine/Source/Runtime/MeshDescription/Private/MeshDescriptionBase.cpp
+++ b/Engine/Source/Runtime/MeshDescription/Private/MeshDescriptionBase.cpp
@@ -85,7 +85,7 @@ FVertexInstanceID UMeshDescriptionBase::CreateVertexInstance(FVertexID VertexID)
 	else
 	{
 		UE_LOG(LogMeshDescription, Warning, TEXT("CreateVertexInstance: VertexID %d doesn't exist."), VertexID.GetValue());
-		return FVertexInstanceID::Invalid;
+		return INDEX_NONE;
 	}
 }
 
@@ -135,13 +135,13 @@ FEdgeID UMeshDescriptionBase::CreateEdge(FVertexID VertexID0, FVertexID VertexID
 	if (!MeshDescription.IsVertexValid(VertexID0))
 	{
 		UE_LOG(LogMeshDescription, Warning, TEXT("CreateEdge: VertexID %d doesn't exist."), VertexID0.GetValue());
-		return FEdgeID::Invalid;
+		return INDEX_NONE;
 	}
 
 	if (!MeshDescription.IsVertexValid(VertexID1))
 	{
 		UE_LOG(LogMeshDescription, Warning, TEXT("CreateEdge: VertexID %d doesn't exist."), VertexID1.GetValue());
-		return FEdgeID::Invalid;
+		return INDEX_NONE;
 	}
 
 	return MeshDescription.CreateEdge(VertexID0, VertexID1);
@@ -197,7 +197,7 @@ FTriangleID UMeshDescriptionBase::CreateTriangle(FPolygonGroupID PolygonGroupID,
 	if (!MeshDescription.IsPolygonGroupValid(PolygonGroupID))
 	{
 		UE_LOG(LogMeshDescription, Warning, TEXT("CreateTriangle: PolygonGroupID %d doesn't exist."), PolygonGroupID.GetValue());
-		return FTriangleID::Invalid;
+		return INDEX_NONE;
 	}
 
 	return MeshDescription.CreateTriangle(PolygonGroupID, VertexInstanceIDs, &NewEdgeIDs);
@@ -246,7 +246,7 @@ FPolygonID UMeshDescriptionBase::CreatePolygon(FPolygonGroupID PolygonGroupID, T
 	if (!MeshDescription.IsPolygonGroupValid(PolygonGroupID))
 	{
 		UE_LOG(LogMeshDescription, Warning, TEXT("CreatePolygon: PolygonGroupID %d doesn't exist."), PolygonGroupID.GetValue());
-		return FPolygonID::Invalid;
+		return INDEX_NONE;
 	}
 
 	return MeshDescription.CreatePolygon(PolygonGroupID, VertexInstanceIDs, &NewEdgeIDs);
@@ -342,13 +342,13 @@ FEdgeID UMeshDescriptionBase::GetVertexPairEdge(FVertexID VertexID0, FVertexID V
 	if (!MeshDescription.IsVertexValid(VertexID0))
 	{
 		UE_LOG(LogMeshDescription, Warning, TEXT("GetVertexPairEdge: VertexID %d doesn't exist."), VertexID0.GetValue());
-		return FEdgeID::Invalid;
+		return INDEX_NONE;
 	}
 
 	if (!MeshDescription.IsVertexValid(VertexID1))
 	{
 		UE_LOG(LogMeshDescription, Warning, TEXT("GetVertexPairEdge: VertexID %d doesn't exist."), VertexID1.GetValue());
-		return FEdgeID::Invalid;
+		return INDEX_NONE;
 	}
 
 	return MeshDescription.GetVertexPairEdge(VertexID0, VertexID1);
@@ -496,7 +496,7 @@ FVertexID UMeshDescriptionBase::GetVertexInstanceVertex(FVertexInstanceID Vertex
 	if (!MeshDescription.IsVertexInstanceValid(VertexInstanceID))
 	{
 		UE_LOG(LogMeshDescription, Warning, TEXT("GetVertexInstanceVertex: VertexInstanceID %d doesn't exist."), VertexInstanceID.GetValue());
-		return FVertexID::Invalid;
+		return INDEX_NONE;
 	}
 
 	return MeshDescription.GetVertexInstanceVertex(VertexInstanceID);
@@ -507,13 +507,13 @@ FEdgeID UMeshDescriptionBase::GetVertexInstancePairEdge(FVertexInstanceID Vertex
 	if (!MeshDescription.IsVertexInstanceValid(VertexInstanceID0))
 	{
 		UE_LOG(LogMeshDescription, Warning, TEXT("GetVertexInstancePairEdge: VertexInstanceID %d doesn't exist."), VertexInstanceID0.GetValue());
-		return FEdgeID::Invalid;
+		return INDEX_NONE;
 	}
 
 	if (!MeshDescription.IsVertexInstanceValid(VertexInstanceID1))
 	{
 		UE_LOG(LogMeshDescription, Warning, TEXT("GetVertexInstancePairEdge: VertexInstanceID %d doesn't exist."), VertexInstanceID1.GetValue());
-		return FEdgeID::Invalid;
+		return INDEX_NONE;
 	}
 
 	return MeshDescription.GetVertexInstancePairEdge(VertexInstanceID0, VertexInstanceID1);
@@ -644,13 +644,13 @@ FVertexID UMeshDescriptionBase::GetEdgeVertex(FEdgeID EdgeID, int32 VertexNumber
 	if (!MeshDescription.IsEdgeValid(EdgeID))
 	{
 		UE_LOG(LogMeshDescription, Warning, TEXT("GetEdgeVertex: EdgeID %d doesn't exist."), EdgeID.GetValue());
-		return FVertexID::Invalid;
+		return INDEX_NONE;
 	}
 
 	if (VertexNumber != 0 && VertexNumber != 1)
 	{
 		UE_LOG(LogMeshDescription, Warning, TEXT("GetEdgeVertex: invalid vertex number %d."), VertexNumber);
-		return FVertexID::Invalid;
+		return INDEX_NONE;
 	}
 
 	return MeshDescription.GetEdgeVertex(EdgeID, VertexNumber);
@@ -676,7 +676,7 @@ FPolygonID UMeshDescriptionBase::GetTrianglePolygon(FTriangleID TriangleID) cons
 	if (!MeshDescription.IsTriangleValid(TriangleID))
 	{
 		UE_LOG(LogMeshDescription, Warning, TEXT("GetTrianglePolygon: TriangleID %d doesn't exist."), TriangleID.GetValue());
-		return FPolygonID::Invalid;
+		return INDEX_NONE;
 	}
 
 	return MeshDescription.GetTrianglePolygon(TriangleID);
@@ -687,7 +687,7 @@ FPolygonGroupID UMeshDescriptionBase::GetTrianglePolygonGroup(FTriangleID Triang
 	if (!MeshDescription.IsTriangleValid(TriangleID))
 	{
 		UE_LOG(LogMeshDescription, Warning, TEXT("GetTrianglePolygonGroup: TriangleID %d doesn't exist."), TriangleID.GetValue());
-		return FPolygonGroupID::Invalid;
+		return INDEX_NONE;
 	}
 
 	return MeshDescription.GetTrianglePolygonGroup(TriangleID);
@@ -720,13 +720,13 @@ FVertexInstanceID UMeshDescriptionBase::GetTriangleVertexInstance(FTriangleID Tr
 	if (!MeshDescription.IsTriangleValid(TriangleID))
 	{
 		UE_LOG(LogMeshDescription, Warning, TEXT("GetTriangleVertexInstance: TriangleID %d doesn't exist."), TriangleID.GetValue());
-		return FVertexInstanceID::Invalid;
+		return INDEX_NONE;
 	}
 
 	if (Index < 0 || Index > 2)
 	{
 		UE_LOG(LogMeshDescription, Warning, TEXT("GetTriangleVertexInstance: invalid vertex index %d."), Index);
-		return FVertexInstanceID::Invalid;
+		return INDEX_NONE;
 	}
 
 	return MeshDescription.GetTriangleVertexInstance(TriangleID, Index);
@@ -741,7 +741,7 @@ void UMeshDescriptionBase::GetTriangleVertices(FTriangleID TriangleID, TArray<FV
 	}
 
 	OutVertexIDs.SetNumUninitialized(3);
-	MeshDescription.GetTriangleVertices(TriangleID, OutVertexIDs);
+	Algo::Copy(MeshDescription.GetTriangleVertices(TriangleID), OutVertexIDs);
 }
 
 void UMeshDescriptionBase::GetTriangleEdges(FTriangleID TriangleID, TArray<FEdgeID>& OutEdgeIDs) const
@@ -753,7 +753,7 @@ void UMeshDescriptionBase::GetTriangleEdges(FTriangleID TriangleID, TArray<FEdge
 	}
 
 	OutEdgeIDs.SetNumUninitialized(3);
-	MeshDescription.GetTriangleEdges(TriangleID, OutEdgeIDs);
+	Algo::Copy(MeshDescription.GetTriangleEdges(TriangleID), OutEdgeIDs);
 }
 
 void UMeshDescriptionBase::GetTriangleAdjacentTriangles(FTriangleID TriangleID, TArray<FTriangleID>& OutTriangleIDs) const
@@ -772,7 +772,7 @@ FVertexInstanceID UMeshDescriptionBase::GetVertexInstanceForTriangleVertex(FTria
 	if (!MeshDescription.IsTriangleValid(TriangleID))
 	{
 		UE_LOG(LogMeshDescription, Warning, TEXT("GetVertexInstanceForTriangleVertex: TriangleID %d doesn't exist."), TriangleID.GetValue());
-		return FVertexInstanceID::Invalid;
+		return INDEX_NONE;
 	}
 
 	return MeshDescription.GetVertexInstanceForTriangleVertex(TriangleID, VertexID);
@@ -886,7 +886,7 @@ FPolygonGroupID UMeshDescriptionBase::GetPolygonPolygonGroup(FPolygonID PolygonI
 	if (!MeshDescription.IsPolygonValid(PolygonID))
 	{
 		UE_LOG(LogMeshDescription, Warning, TEXT("GetPolygonPolygonGroup: PolygonID %d doesn't exist."), PolygonID.GetValue());
-		return FPolygonGroupID::Invalid;
+		return INDEX_NONE;
 	}
 
 	return MeshDescription.GetPolygonPolygonGroup(PolygonID);
@@ -897,33 +897,30 @@ FVertexInstanceID UMeshDescriptionBase::GetVertexInstanceForPolygonVertex(FPolyg
 	if (!MeshDescription.IsPolygonValid(PolygonID))
 	{
 		UE_LOG(LogMeshDescription, Warning, TEXT("GetVertexInstanceForPolygonVertex: PolygonID %d doesn't exist."), PolygonID.GetValue());
-		return FVertexInstanceID::Invalid;
+		return INDEX_NONE;
 	}
 
 	return MeshDescription.GetVertexInstanceForPolygonVertex(PolygonID, VertexID);
 }
 
-void UMeshDescriptionBase::SetPolygonVertexInstance(FPolygonID PolygonID, int32 PerimeterIndex, FVertexInstanceID VertexInstanceID)
+void UMeshDescriptionBase::SetPolygonVertexInstances(FPolygonID PolygonID, const TArray<FVertexInstanceID>& VertexInstanceIDs)
 {
 	if (!MeshDescription.IsPolygonValid(PolygonID))
 	{
-		UE_LOG(LogMeshDescription, Warning, TEXT("SetPolygonVertexInstance: PolygonID %d doesn't exist."), PolygonID.GetValue());
+		UE_LOG(LogMeshDescription, Warning, TEXT("SetPolygonVertexInstances: PolygonID %d doesn't exist."), PolygonID.GetValue());
 		return;
 	}
 
-	if (PerimeterIndex >= MeshDescription.GetNumPolygonVertices(PolygonID))
+	for (FVertexInstanceID VertexInstanceID : VertexInstanceIDs)
 	{
-		UE_LOG(LogMeshDescription, Warning, TEXT("SetPolygonVertexInstance: Out of range vertex index %d."), PerimeterIndex);
-		return;
-	}
-
-	if (!MeshDescription.IsVertexInstanceValid(VertexInstanceID))
-	{
-		UE_LOG(LogMeshDescription, Warning, TEXT("SetPolygonVertexInstance: VertexInstanceID %d doesn't exist."), VertexInstanceID.GetValue());
-		return;
+		if (!MeshDescription.IsVertexInstanceValid(VertexInstanceID))
+		{
+			UE_LOG(LogMeshDescription, Warning, TEXT("SetPolygonVertexInstances: VertexInstanceID %d doesn't exist."), VertexInstanceID.GetValue());
+			return;
+		}
 	}
 
-	MeshDescription.SetPolygonVertexInstance(PolygonID, PerimeterIndex, VertexInstanceID);
+	MeshDescription.SetPolygonVertexInstances(PolygonID, VertexInstanceIDs);
 }
 
 void UMeshDescriptionBase::SetPolygonPolygonGroup(FPolygonID PolygonID, FPolygonGroupID PolygonGroupID)
