First pass of MeshDescription API and format refactor.

diff --git a/Engine/Source/Runtime/MeshDescription/Private/MeshTypes.cpp b/Engine/Source/Runtime/MeshDescription/Private/MeshTypes.cpp
--- a/Engine/Source/Runtime/MeshDescription/Private/MeshTypes.cpp
+++ b/Engine/Source/Runtime/MeshDescription/Private/MeshTypes.cpp
@@ -5,13 +5,13 @@
 DEFINE_LOG_CATEGORY( LogMeshDescription );
 
 
-const FElementID FElementID::Invalid( TNumericLimits<uint32>::Max() );
-const FVertexID FVertexID::Invalid( TNumericLimits<uint32>::Max() );
-const FVertexInstanceID FVertexInstanceID::Invalid( TNumericLimits<uint32>::Max() );
-const FEdgeID FEdgeID::Invalid( TNumericLimits<uint32>::Max() );
-const FTriangleID FTriangleID::Invalid(TNumericLimits<uint32>::Max());
-const FPolygonGroupID FPolygonGroupID::Invalid( TNumericLimits<uint32>::Max() );
-const FPolygonID FPolygonID::Invalid( TNumericLimits<uint32>::Max() );
+const FElementID FElementID::Invalid( INDEX_NONE );
+const FVertexID FVertexID::Invalid( INDEX_NONE );
+const FVertexInstanceID FVertexInstanceID::Invalid( INDEX_NONE );
+const FEdgeID FEdgeID::Invalid( INDEX_NONE );
+const FTriangleID FTriangleID::Invalid( INDEX_NONE );
+const FPolygonGroupID FPolygonGroupID::Invalid( INDEX_NONE );
+const FPolygonID FPolygonID::Invalid( INDEX_NONE );
 
 #if 0
 const FName UEditableMeshAttribute::VertexPositionName( "VertexPosition" );
diff --git a/Engine/Source/Runtime/MeshDescription/Public/MeshAttributeArray.h b/Engine/Source/Runtime/MeshDescription/Public/MeshAttributeArray.h
--- a/Engine/Source/Runtime/MeshDescription/Public/MeshAttributeArray.h
+++ b/Engine/Source/Runtime/MeshDescription/Public/MeshAttributeArray.h
@@ -3,11 +3,14 @@
 #pragma once
 
 #include "CoreMinimal.h"
+#include "Containers/ArrayView.h"
 #include "MeshTypes.h"
 #include "MeshElementRemappings.h"
 #include "Misc/TVariant.h"
 #include "Templates/CopyQualifiersFromTo.h"
+#include "UObject/EditorObjectVersion.h"
 #include "UObject/ReleaseObjectVersion.h"
+#include "Templates/IsArray.h"
 
 
 /**
@@ -65,6 +68,13 @@ template <typename Tuple> struct TVariantFromTuple;
 template <typename... Ts> struct TVariantFromTuple<TTuple<Ts...>> { using Type = TVariant<FEmptyVariantState, Ts...>; };
 
 
+/**
+ * Helper template which, given an array type, can break into its element type and number of elements.
+ */
+template <typename T> struct TBreakArray;
+template <typename T, SIZE_T N> struct TBreakArray<T[N]> : TIntegralConstant<uint32, N> { using Type = T; };
+
+
 /**
  * Class which implements a function jump table to be automatically generated at compile time.
  * This is used by TAttributesSet to provide O(1) dispatch by attribute type at runtime.
@@ -97,58 +107,87 @@ template <typename AttributeType>
 class TMeshAttributeArrayBase
 {
 public:
+	explicit TMeshAttributeArrayBase(uint32 InExtent = 1)
+		: Extent(InExtent)
+	{}
 
 	/** Custom serialization for TMeshAttributeArrayBase. */
 	template <typename T> friend typename TEnableIf<!TIsBulkSerializable<T>::Value, FArchive>::Type& operator<<( FArchive& Ar, TMeshAttributeArrayBase<T>& Array );
 	template <typename T> friend typename TEnableIf<TIsBulkSerializable<T>::Value, FArchive>::Type& operator<<( FArchive& Ar, TMeshAttributeArrayBase<T>& Array );
 
 	/** Return size of container */
-	FORCEINLINE int32 Num() const { return Container.Num(); }
+	FORCEINLINE int32 Num() const { return Container.Num() / Extent; }
 
 	/** Return base of data */
+	UE_DEPRECATED(4.25, "This method will be removed.")
 	FORCEINLINE const AttributeType* GetData() const { return Container.GetData(); }
 
 	/** Initializes the array to the given size with the default value */
-	FORCEINLINE void Initialize( const int32 ElementCount, const AttributeType& Default )
+	FORCEINLINE void Initialize(const int32 ElementCount, const AttributeType& Default)
 	{
-		Container.Reset( ElementCount );
-		Insert( ElementCount - 1, Default );
+		Container.Reset(ElementCount * Extent);
+		if (ElementCount > 0)
+		{
+			Insert(ElementCount - 1, Default);
+		}
 	}
 
-	uint32 GetHash( uint32 Crc = 0 ) const
+	uint32 GetHash(uint32 Crc = 0) const
 	{
-		return FCrc::MemCrc32( Container.GetData(), Container.Num() * sizeof( AttributeType ), Crc );
+		return FCrc::MemCrc32(Container.GetData(), Container.Num() * sizeof(AttributeType), Crc);
 	}
 
 	/** Expands the array if necessary so that the passed element index is valid. Newly created elements will be assigned the default value. */
-	void Insert( const int32 Index, const AttributeType& Default );
+	void Insert(const int32 Index, const AttributeType& Default);
+
+	/** Fills the index with the default value */
+	void SetToDefault(const int32 Index, const AttributeType& Default)
+	{
+		for (uint32 I = 0; I < Extent; ++I)
+		{
+			Container[Index * Extent + I] = Default;
+		}
+	}
 
 	/** Remaps elements according to the passed remapping table */
-	void Remap( const TSparseArray<int32>& IndexRemap, const AttributeType& Default );
+	void Remap(const TSparseArray<int32>& IndexRemap, const AttributeType& Default);
 
 	/** Element accessors */
-	FORCEINLINE const AttributeType& operator[]( const int32 Index ) const { return Container[ Index ]; }
-	FORCEINLINE AttributeType& operator[]( const int32 Index ) { return Container[ Index ]; }
+	UE_DEPRECATED(4.25, "Please use GetElementBase() instead.")
+	FORCEINLINE const AttributeType& operator[](const int32 Index) const { return Container[Index]; }
+
+	UE_DEPRECATED(4.25, "Please use GetElementBase() instead.")
+	FORCEINLINE AttributeType& operator[](const int32 Index) { return Container[Index]; }
+
+	FORCEINLINE const AttributeType* GetElementBase(const int32 Index) const { return &Container[Index * Extent]; }
+	FORCEINLINE AttributeType* GetElementBase(const int32 Index) { return &Container[Index * Extent]; }
+
+	FORCEINLINE uint32 GetExtent() const { return Extent; }
 
 protected:
 	/** The actual container, represented by a regular array */
 	TArray<AttributeType> Container;
+
+	/** Number of array elements in this attribute type */
+	uint32 Extent;
 };
 
 
 template <typename AttributeType>
-void TMeshAttributeArrayBase<AttributeType>::Insert( const int32 Index, const AttributeType& Default )
+void TMeshAttributeArrayBase<AttributeType>::Insert(const int32 Index, const AttributeType& Default)
 {
-	if( Index >= Container.Num() )
+	int32 EndIndex = (Index + 1) * Extent;
+	if (EndIndex > Container.Num())
 	{
 		// If the index is off the end of the container, add as many elements as required to make it the last valid index.
-		int32 StartIndex = Container.AddUninitialized( Index + 1 - Container.Num() );
+		int32 StartIndex = Container.AddUninitialized(EndIndex - Container.Num());
 		AttributeType* Data = Container.GetData() + StartIndex;
 
 		// Construct added elements with the default value passed in
-		while( StartIndex <= Index )
+
+		while (StartIndex < EndIndex)
 		{
-			new( Data ) AttributeType( Default );
+			new(Data) AttributeType(Default);
 			StartIndex++;
 			Data++;
 		}
@@ -156,25 +195,39 @@ void TMeshAttributeArrayBase<AttributeType>::Insert( const int32 Index, const At
 }
 
 template <typename AttributeType>
-void TMeshAttributeArrayBase<AttributeType>::Remap( const TSparseArray<int32>& IndexRemap, const AttributeType& Default )
+void TMeshAttributeArrayBase<AttributeType>::Remap(const TSparseArray<int32>& IndexRemap, const AttributeType& Default)
 {
-	TMeshAttributeArrayBase NewAttributeArray;
+	TMeshAttributeArrayBase NewAttributeArray(Extent);
 
-	for( typename TSparseArray<int32>::TConstIterator It( IndexRemap ); It; ++It )
+	for (typename TSparseArray<int32>::TConstIterator It(IndexRemap); It; ++It)
 	{
 		const int32 OldElementIndex = It.GetIndex();
-		const int32 NewElementIndex = IndexRemap[ OldElementIndex ];
+		const int32 NewElementIndex = IndexRemap[OldElementIndex];
 
-		NewAttributeArray.Insert( NewElementIndex, Default );
-		NewAttributeArray[ NewElementIndex ] = MoveTemp( Container[ OldElementIndex ] );
+		NewAttributeArray.Insert(NewElementIndex, Default);
+		AttributeType* DestElementBase = NewAttributeArray.GetElementBase(NewElementIndex);
+		AttributeType* SrcElementBase = GetElementBase(OldElementIndex);
+		for (uint32 Index = 0; Index < Extent; ++Index)
+		{
+			DestElementBase[Index] = MoveTemp(SrcElementBase[Index]);
+		}
 	}
 
-	Container = MoveTemp( NewAttributeArray.Container );
+	Container = MoveTemp(NewAttributeArray.Container);
 }
 
 template <typename T>
 inline typename TEnableIf<!TIsBulkSerializable<T>::Value, FArchive>::Type& operator<<( FArchive& Ar, TMeshAttributeArrayBase<T>& Array )
 {
+	if (Ar.IsLoading() && Ar.CustomVer(FEditorObjectVersion::GUID) < FReleaseObjectVersion::MeshDescriptionNewFormat)
+	{
+		Array.Extent = 1;
+	}
+	else
+	{
+		Ar << Array.Extent;
+	}
+
 	// Serialize types which aren't bulk serializable, which need to be serialized element-by-element
 	Ar << Array.Container;
 	return Ar;
@@ -183,6 +236,15 @@ inline typename TEnableIf<!TIsBulkSerializable<T>::Value, FArchive>::Type& opera
 template <typename T>
 inline typename TEnableIf<TIsBulkSerializable<T>::Value, FArchive>::Type& operator<<( FArchive& Ar, TMeshAttributeArrayBase<T>& Array )
 {
+	if (Ar.IsLoading() && Ar.CustomVer(FEditorObjectVersion::GUID) < FReleaseObjectVersion::MeshDescriptionNewFormat)
+	{
+		Array.Extent = 1;
+	}
+	else
+	{
+		Ar << Array.Extent;
+	}
+
 	if( Ar.IsLoading() && Ar.CustomVer( FReleaseObjectVersion::GUID ) < FReleaseObjectVersion::MeshDescriptionNewSerialization )
 	{
 		// Legacy path for old format attribute arrays. BulkSerialize has a different format from regular serialization.
@@ -198,37 +260,20 @@ inline typename TEnableIf<TIsBulkSerializable<T>::Value, FArchive>::Type& operat
 }
 
 
-// This is a deprecated class which will be removed in 4.21
-template <typename AttributeType, typename ElementIDType>
-class TMeshAttributeArray final : public TMeshAttributeArrayBase<AttributeType>
-{
-	static_assert( TIsDerivedFrom<ElementIDType, FElementID>::IsDerived, "ElementIDType must be derived from FElementID" );
-
-	using Super = TMeshAttributeArrayBase<AttributeType>;
-	using Super::operator[];
-
-public:
-
-	/** Element accessors */
-	FORCEINLINE const AttributeType& operator[]( const ElementIDType ElementID ) const { return this->Container[ ElementID.GetValue() ]; }
-	FORCEINLINE AttributeType& operator[]( const ElementIDType ElementID ) { return this->Container[ ElementID.GetValue() ]; }
-};
-
-
-
 /**
  * Flags specifying properties of an attribute
  */
 enum class EMeshAttributeFlags : uint32
 {
 	None				= 0,
-	Lerpable			= ( 1 << 0 ),	/** Attribute can be automatically lerped according to the value of 2 or 3 other attributes */
-	AutoGenerated		= ( 1 << 1 ),	/** Attribute is auto-generated by importer or editable mesh, rather than representing an imported property */	
-	Mergeable			= ( 1 << 2 ),	/** If all vertices' attributes are mergeable, and of near-equal value, they can be welded */
-	Transient			= ( 1 << 3 )	/** Attribute is not serialized */
+	Lerpable			= (1 << 0),		/** Attribute can be automatically lerped according to the value of 2 or 3 other attributes */
+	AutoGenerated		= (1 << 1),		/** Attribute is auto-generated by importer or editable mesh, rather than representing an imported property */	
+	Mergeable			= (1 << 2),		/** If all vertices' attributes are mergeable, and of near-equal value, they can be welded */
+	Transient			= (1 << 3),		/** Attribute is not serialized */
+	IndexReference		= (1 << 4)		/** Attribute is a reference to another element index */
 };
 
-ENUM_CLASS_FLAGS( EMeshAttributeFlags );
+ENUM_CLASS_FLAGS(EMeshAttributeFlags);
 
 
 /**
@@ -247,25 +292,36 @@ class FMeshAttributeArraySetBase
 {
 public:
 	/** Constructor */
-	FORCEINLINE FMeshAttributeArraySetBase( const uint32 InType, const EMeshAttributeFlags InFlags, const int32 InNumberOfElements )
-		: Type( InType ),
-		  NumElements( InNumberOfElements ),
-		  Flags( InFlags )
+	FORCEINLINE FMeshAttributeArraySetBase(const uint32 InType, const EMeshAttributeFlags InFlags, const int32 InNumberOfElements, const uint32 InExtent)
+		: Type(InType),
+		  Extent(InExtent),
+		  NumElements(InNumberOfElements),
+		  Flags(InFlags)
 	{}
 
 	/** Virtual interface */
 	virtual ~FMeshAttributeArraySetBase() = default;
 	virtual TUniquePtr<FMeshAttributeArraySetBase> Clone() const = 0;
-	virtual void Insert( const int32 Index ) = 0;
-	virtual void Remove( const int32 Index ) = 0;
-	virtual void Initialize( const int32 Count ) = 0;
+	virtual void Insert(const int32 Index) = 0;
+	virtual void Remove(const int32 Index) = 0;
+	virtual void Initialize(const int32 Count) = 0;
 	virtual uint32 GetHash() const = 0;
-	virtual void Serialize( FArchive& Ar ) = 0;
-	virtual void Remap( const TSparseArray<int32>& IndexRemap ) = 0;
+	virtual void Serialize(FArchive& Ar) = 0;
+	virtual void Remap(const TSparseArray<int32>& IndexRemap) = 0;
+
+//	UE_DEPRECATED(4.26, "Please use GetNumChannels().")
 	virtual int32 GetNumIndices() const = 0;
-	virtual void SetNumIndices( const int32 NumIndices ) = 0;
-	virtual void InsertIndex( const int32 Index ) = 0;
-	virtual void RemoveIndex( const int32 Index ) = 0;
+//	UE_DEPRECATED(4.26, "Please use SetNumChannels().")
+	virtual void SetNumIndices(const int32 NumIndices) = 0;
+//	UE_DEPRECATED(4.26, "Please use InsertChannel().")
+	virtual void InsertIndex(const int32 Index) = 0;
+//	UE_DEPRECATED(4.26, "Please use RemoveChannel().")
+	virtual void RemoveIndex(const int32 Index) = 0;
+
+	virtual int32 GetNumChannels() const = 0;
+	virtual void SetNumChannels(const int32 NumChannels) = 0;
+	virtual void InsertChannel(const int32 Index) = 0;
+	virtual void RemoveChannel(const int32 Index) = 0;
 
 	/** Determine whether this attribute array set is of the given type */
 	template <typename T>
@@ -277,11 +333,14 @@ public:
 	/** Get the type index of this attribute array set */
 	FORCEINLINE uint32 GetType() const { return Type; }
 
+	/** Get the type extent of this attribute array set */
+	FORCEINLINE uint32 GetExtent() const { return Extent; }
+
 	/** Get the flags for this attribute array set */
 	FORCEINLINE EMeshAttributeFlags GetFlags() const { return Flags; }
 
 	/** Set the flags for this attribute array set */
-	FORCEINLINE void SetFlags( const EMeshAttributeFlags InFlags ) { Flags = InFlags; }
+	FORCEINLINE void SetFlags(const EMeshAttributeFlags InFlags) { Flags = InFlags; }
 
 	/** Return number of elements each attribute index has */
 	FORCEINLINE int32 GetNumElements() const { return NumElements; }
@@ -290,6 +349,9 @@ protected:
 	/** Type of the attribute array (based on the tuple element index from AttributeTypes) */
 	uint32 Type;
 
+	/** Extent of the type, i.e. the number of array elements it consists of */
+	uint32 Extent;
+
 	/** Number of elements in each index */
 	int32 NumElements;
 
@@ -307,209 +369,276 @@ class TMeshAttributeArraySet : public FMeshAttributeArraySetBase
 	using Super = FMeshAttributeArraySetBase;
 
 public:
-	/** Constructor */
-	FORCEINLINE explicit TMeshAttributeArraySet( const int32 NumberOfIndices = 0, const AttributeType& InDefaultValue = AttributeType(), const EMeshAttributeFlags InFlags = EMeshAttributeFlags::None, const int32 InNumberOfElements = 0 )
-		: Super( TTupleIndex<AttributeType, AttributeTypes>::Value, InFlags, InNumberOfElements ),
-		  DefaultValue( InDefaultValue )
+	/** Constructors */
+	FORCEINLINE explicit TMeshAttributeArraySet(const int32 Extent = 1)
+		: Super(TTupleIndex<AttributeType, AttributeTypes>::Value, EMeshAttributeFlags::None, 0, Extent)
+	{}
+
+	FORCEINLINE explicit TMeshAttributeArraySet(const int32 NumberOfChannels, const AttributeType& InDefaultValue, const EMeshAttributeFlags InFlags, const int32 InNumberOfElements, const uint32 Extent)
+		: Super(TTupleIndex<AttributeType, AttributeTypes>::Value, InFlags, InNumberOfElements, Extent),
+		  DefaultValue(InDefaultValue)
 	{
-		SetNumIndices( NumberOfIndices );
+		SetNumChannels(NumberOfChannels);
 	}
 
 	/** Creates a copy of itself and returns a TUniquePtr to it */
-	virtual TUniquePtr<FMeshAttributeArraySetBase> Clone() const
+	virtual TUniquePtr<FMeshAttributeArraySetBase> Clone() const override
 	{
-		return MakeUnique<TMeshAttributeArraySet>( *this );
+		return MakeUnique<TMeshAttributeArraySet>(*this);
 	}
 
 	/** Insert the element at the given index */
-	virtual void Insert( const int32 Index )
+	virtual void Insert(const int32 Index) override
 	{
-		for( TMeshAttributeArrayBase<AttributeType>& ArrayForIndex : ArrayForIndices )
+		for (TMeshAttributeArrayBase<AttributeType>& ArrayForChannel : ArrayForChannels)
 		{
-			ArrayForIndex.Insert( Index, DefaultValue );
+			ArrayForChannel.Insert(Index, DefaultValue);
 		}
 
-		NumElements = FMath::Max( NumElements, Index + 1 );
+		NumElements = FMath::Max(NumElements, Index + 1);
 	}
 
 	/** Remove the element at the given index, replacing it with a default value */
-	virtual void Remove( const int32 Index )
+	virtual void Remove(const int32 Index) override
 	{
-		for( TMeshAttributeArrayBase<AttributeType>& ArrayForIndex : ArrayForIndices )
+		for (TMeshAttributeArrayBase<AttributeType>& ArrayForChannel : ArrayForChannels)
 		{
-			ArrayForIndex[ Index ] = DefaultValue;
+			ArrayForChannel.SetToDefault(Index, DefaultValue);
 		}
 	}
 
 	/** Sets the number of elements to the exact number provided, and initializes them to the default value */
-	virtual void Initialize( const int32 Count )
+	virtual void Initialize(const int32 Count) override
 	{
 		NumElements = Count;
-		for( TMeshAttributeArrayBase<AttributeType>& ArrayForIndex : ArrayForIndices )
+		for (TMeshAttributeArrayBase<AttributeType>& ArrayForChannel : ArrayForChannels)
 		{
-			ArrayForIndex.Initialize( Count, DefaultValue );
+			ArrayForChannel.Initialize(Count, DefaultValue);
 		}
 	}
 
-	virtual uint32 GetHash() const
+	virtual uint32 GetHash() const override
 	{
 		uint32 CrcResult = 0;
-		for ( const TMeshAttributeArrayBase<AttributeType>& ArrayForIndex : ArrayForIndices )
+		for (const TMeshAttributeArrayBase<AttributeType>& ArrayForChannel : ArrayForChannels)
 		{
-			CrcResult = ArrayForIndex.GetHash( CrcResult );
+			CrcResult = ArrayForChannel.GetHash(CrcResult);
 		}
 		return CrcResult;
 	}
 
 	/** Polymorphic serialization */
-	virtual void Serialize( FArchive& Ar )
+	virtual void Serialize(FArchive& Ar) override
 	{
-		Ar << ( *this );
+		Ar << (*this);
 	}
 
 	/** Performs an element index remap according to the passed array */
-	virtual void Remap( const TSparseArray<int32>& IndexRemap )
+	virtual void Remap(const TSparseArray<int32>& IndexRemap) override
 	{
-		for( TMeshAttributeArrayBase<AttributeType>& ArrayForIndex : ArrayForIndices )
+		for (TMeshAttributeArrayBase<AttributeType>& ArrayForChannel : ArrayForChannels)
 		{
-			ArrayForIndex.Remap( IndexRemap, DefaultValue );
-			NumElements = ArrayForIndex.Num();
+			ArrayForChannel.Remap(IndexRemap, DefaultValue);
+			NumElements = ArrayForChannel.Num();
 		}
 	}
 
-	/** Return number of indices this attribute has */
-	virtual inline int32 GetNumIndices() const { return ArrayForIndices.Num(); }
+	UE_DEPRECATED(4.26, "Please use GetNumChannels().")
+	virtual inline int32 GetNumIndices() const override { return GetNumChannels(); }
 
-	/** Sets number of indices this attribute has */
-	virtual void SetNumIndices( const int32 NumIndices )
+	/** Return number of channels this attribute has */
+	virtual inline int32 GetNumChannels() const override { return ArrayForChannels.Num(); }
+
+	UE_DEPRECATED(4.26, "Please use SetNumChannels().")
+	virtual void SetNumIndices(const int32 NumIndices) override { SetNumChannels(NumIndices); }
+
+	/** Sets number of channels this attribute has */
+	virtual void SetNumChannels(const int32 NumChannels) override
 	{
-		const int32 OriginalNumIndices = ArrayForIndices.Num();
-		ArrayForIndices.SetNum( NumIndices );
+		if (NumChannels < ArrayForChannels.Num())
+		{
+			ArrayForChannels.SetNum(NumChannels);
+			return;
+		}
 
-		// If we have added new indices, ensure they are filled out with the correct number of elements
-		for( int32 Index = OriginalNumIndices; Index < NumIndices; ++Index )
+		while (ArrayForChannels.Num() < NumChannels)
 		{
-			ArrayForIndices[ Index ].Initialize( NumElements, DefaultValue );
+			TMeshAttributeArrayBase<AttributeType>& Array = ArrayForChannels.Emplace_GetRef(Extent);
+			Array.Initialize(NumElements, DefaultValue);
 		}
 	}
 
-	/** Insert a new attribute index */
-	virtual void InsertIndex( const int32 Index )
+	UE_DEPRECATED(4.26, "Please use InsertChannel().")
+	virtual void InsertIndex(const int32 Index) override
+	{
+		InsertChannel(Index);
+	}
+
+	/** Insert a new attribute channel */
+	virtual void InsertChannel(const int32 Index) override
+	{
+		TMeshAttributeArrayBase<AttributeType>& Array = ArrayForChannels.EmplaceAt_GetRef(Index, Extent);
+		Array.Initialize(NumElements, DefaultValue);
+	}
+
+	UE_DEPRECATED(4.26, "Please use RemoveChannel().")
+	virtual void RemoveIndex(const int32 Index) override
 	{
-		ArrayForIndices.InsertDefaulted( Index );
-		ArrayForIndices[ Index ].Initialize( NumElements, DefaultValue );
+		RemoveChannel(Index);
 	}
 
-	/** Remove the array at the given index */
-	virtual void RemoveIndex( const int32 Index )
+	/** Remove the channel at the given index */
+	virtual void RemoveChannel(const int32 Index) override
 	{
-		ArrayForIndices.RemoveAt( Index );
+		ArrayForChannels.RemoveAt(Index);
 	}
 
 
-	/** Return the TMeshAttributeArrayBase corresponding to the given attribute index */
-	FORCEINLINE const TMeshAttributeArrayBase<AttributeType>& GetArrayForIndex( const int32 Index ) const { return ArrayForIndices[ Index ]; }
-	FORCEINLINE TMeshAttributeArrayBase<AttributeType>& GetArrayForIndex( const int32 Index ) { return ArrayForIndices[ Index ]; }
+	UE_DEPRECATED(4.26, "Please use GetArrayForChannel().")
+	FORCEINLINE const TMeshAttributeArrayBase<AttributeType>& GetArrayForIndex( const int32 Index ) const { return ArrayForChannels[ Index ]; }
+	UE_DEPRECATED(4.26, "Please use GetArrayForChannel().")
+	FORCEINLINE TMeshAttributeArrayBase<AttributeType>& GetArrayForIndex( const int32 Index ) { return ArrayForChannels[ Index ]; }
+
+	/** Return the TMeshAttributeArrayBase corresponding to the given attribute channel */
+	FORCEINLINE const TMeshAttributeArrayBase<AttributeType>& GetArrayForChannel( const int32 Index ) const { return ArrayForChannels[ Index ]; }
+	FORCEINLINE TMeshAttributeArrayBase<AttributeType>& GetArrayForChannel( const int32 Index ) { return ArrayForChannels[ Index ]; }
 
 	/** Return default value for this attribute type */
 	FORCEINLINE AttributeType GetDefaultValue() const { return DefaultValue; }
 
 	/** Serializer */
-	friend FArchive& operator<<( FArchive& Ar, TMeshAttributeArraySet& AttributeArraySet )
+	friend FArchive& operator<<(FArchive& Ar, TMeshAttributeArraySet& AttributeArraySet)
 	{
 		Ar << AttributeArraySet.NumElements;
-		Ar << AttributeArraySet.ArrayForIndices;
+		Ar << AttributeArraySet.ArrayForChannels;
 		Ar << AttributeArraySet.DefaultValue;
 		Ar << AttributeArraySet.Flags;
+
 		return Ar;
 	}
 
 protected:
-	/** An array of MeshAttributeArrays, one per attribute index */
-	TArray<TMeshAttributeArrayBase<AttributeType>> ArrayForIndices;
+	/** An array of MeshAttributeArrays, one per channel */
+	TArray<TMeshAttributeArrayBase<AttributeType>, TInlineAllocator<1>> ArrayForChannels;
 
 	/** The default value for an attribute of this name */
 	AttributeType DefaultValue;
 };
 
 
-// This is a deprecated class which will be removed in 4.21
-template <typename T, typename U>
-class TAttributeIndicesArray : public TMeshAttributeArraySet<T>
-{
-public:
+template <typename> struct TIsArrayView { enum { Value = false }; };
+template <typename T> struct TIsArrayView<TArrayView<T>> { enum { Value = true }; };
 
-	using AttributeType = T;
-	using ElementIDType = U;
 
-	/** Return the TMeshAttributeArray corresponding to the given attribute index */
-	UE_DEPRECATED( 4.20, "Please use GetAttributesRef() or GetAttributesView() instead." )
-	FORCEINLINE const TMeshAttributeArray<AttributeType, ElementIDType>& GetArrayForIndex( const int32 Index ) const
-	{
-		return static_cast<const TMeshAttributeArray<AttributeType, ElementIDType>&>( this->ArrayForIndices[ Index ] );
-	}
+template <typename> struct TBreakArrayView;
+template <typename T> struct TBreakArrayView<TArrayView<T>> { using Type = T; };
 
-	UE_DEPRECATED( 4.20, "Please use GetAttributesRef() or GetAttributesView() instead." )
-	FORCEINLINE TMeshAttributeArray<AttributeType, ElementIDType>& GetArrayForIndex( const int32 Index )
-	{
-		return static_cast<TMeshAttributeArray<AttributeType, ElementIDType>&>( this->ArrayForIndices[ Index ] );
-	}
-};
 
+template <typename T, bool IsArray = TIsArrayView<T>::Value> struct TConstAttribute;
+template <typename T> struct TConstAttribute<T, false> { using Type = const T; };
+template <typename T> struct TConstAttribute<TArrayView<T>, true> { using Type = TArrayView<const T>; };
 
 /**
  * This is the class used to access attribute values.
  * It is a proxy object to a TMeshAttributeArraySet<> and should be passed by value.
  * It is valid for as long as the owning FMeshDescription exists.
  */
+template <typename ElementIDType, typename AttributeType, bool IsArray = TIsArrayView<AttributeType>::Value>
+class TMeshAttributesRef;
+
+template <typename ElementIDType, typename AttributeType, bool IsArray = TIsArrayView<AttributeType>::Value>
+using TMeshAttributesConstRef = TMeshAttributesRef<ElementIDType, typename TConstAttribute<AttributeType>::Type, IsArray>;
+
+template <typename AttributeType, bool IsArray = TIsArrayView<AttributeType>::Value>
+using TMeshAttributesArray = TMeshAttributesRef<int32, AttributeType, IsArray>;
+
+template <typename AttributeType, bool IsArray = TIsArrayView<AttributeType>::Value>
+using TMeshAttributesConstArray = TMeshAttributesRef<int32, typename TConstAttribute<AttributeType>::Type, IsArray>;
+
+
 template <typename ElementIDType, typename AttributeType>
-class TMeshAttributesRef
+class TMeshAttributesRef<ElementIDType, AttributeType, false>
 {
-	template <typename T, typename U> friend class TMeshAttributesRef;
+	template <typename T, typename U, bool IsArray> friend class TMeshAttributesRef;
 
 public:
 	using Type = AttributeType;
 	using ArrayType = typename TCopyQualifiersFromTo<AttributeType, TMeshAttributeArraySet<typename TRemoveCV<AttributeType>::Type>>::Type;
 
 	/** Constructor taking a pointer to a TMeshAttributeArraySet */
-	FORCEINLINE explicit TMeshAttributesRef( ArrayType* InArrayPtr = nullptr )
-		: ArrayPtr( InArrayPtr )
+	FORCEINLINE explicit TMeshAttributesRef(ArrayType* InArrayPtr = nullptr)
+		: ArrayPtr(InArrayPtr)
 	{}
 
 	/** Implicitly construct a TMeshAttributesRef-to-const from a regular one */
 	template <typename T = AttributeType, typename TEnableIf<TIsSame<T, const T>::Value, int>::Type = 0>
-	FORCEINLINE TMeshAttributesRef( TMeshAttributesRef<ElementIDType, typename TRemoveCV<T>::Type> InRef )
-		: ArrayPtr( InRef.ArrayPtr )
+	FORCEINLINE TMeshAttributesRef(TMeshAttributesRef<ElementIDType, typename TRemoveCV<T>::Type> InRef)
+		: ArrayPtr(InRef.ArrayPtr)
 	{}
 
-	/** Access elements from attribute index 0 */
-	FORCEINLINE AttributeType& operator[]( const ElementIDType ElementID ) const
+	/** Access elements from attribute channel 0 */
+	template <typename T = ElementIDType, typename TEnableIf<TIsDerivedFrom<T, FElementID>::Value, int>::Type = 0>
+	FORCEINLINE AttributeType& operator[](const ElementIDType ElementID) const
 	{
-		return ArrayPtr->GetArrayForIndex( 0 )[ ElementID.GetValue() ];
+		return ArrayPtr->GetArrayForChannel(0).GetElementBase(ElementID.GetValue())[0];
 	}
 
-	/** Get the element with the given ID from index 0 */
-	FORCEINLINE AttributeType Get( const ElementIDType ElementID ) const
+	/** Get the element with the given ID from channel 0 */
+	template <typename T = ElementIDType, typename TEnableIf<TIsDerivedFrom<T, FElementID>::Value, int>::Type = 0>
+	FORCEINLINE AttributeType Get(const ElementIDType ElementID) const
+	{
+		return ArrayPtr->GetArrayForChannel(0).GetElementBase(ElementID.GetValue())[0];
+	}
+
+	/** Get the element with the given ID and channel */
+	template <typename T = ElementIDType, typename TEnableIf<TIsDerivedFrom<T, FElementID>::Value, int>::Type = 0>
+	FORCEINLINE AttributeType Get(const ElementIDType ElementID, const int32 Channel) const
 	{
-		return ArrayPtr->GetArrayForIndex( 0 )[ ElementID.GetValue() ];
+		return ArrayPtr->GetArrayForChannel(Channel).GetElementBase(ElementID.GetValue())[0];
 	}
 
-	/** Get the element with the given ID and index */
-	FORCEINLINE AttributeType Get( const ElementIDType ElementID, const int32 Index ) const
+	FORCEINLINE AttributeType& operator[](int32 ElementIndex) const
 	{
-		return ArrayPtr->GetArrayForIndex( Index )[ ElementID.GetValue() ];
+		return ArrayPtr->GetArrayForChannel(0).GetElementBase(ElementIndex)[0];
+	}
+
+	FORCEINLINE AttributeType Get(int32 ElementIndex) const
+	{
+		return ArrayPtr->GetArrayForChannel(0).GetElementBase(ElementIndex)[0];
+	}
+
+	FORCEINLINE AttributeType Get(int32 ElementIndex, const int32 Index) const
+	{
+		return ArrayPtr->GetArrayForChannel(Index).GetElementBase(ElementIndex)[0];
+	}
+
+	FORCEINLINE TArrayView<Type> GetRawArray(const int32 AttributeChannel = 0) const
+	{
+		if (ArrayPtr == nullptr)
+		{
+			return TArrayView<Type>();
+		}
+
+		Type* Element = ArrayPtr->GetArrayForChannel(AttributeChannel).GetElementBase(0);
+		return TArrayView<Type>(Element, GetNumElements());
 	}
 
 	/** Return whether the reference is valid or not */
-	FORCEINLINE bool IsValid() const { return ( ArrayPtr != nullptr ); }
+	FORCEINLINE bool IsValid() const { return (ArrayPtr != nullptr); }
 
 	/** Return default value for this attribute type */
-	FORCEINLINE AttributeType GetDefaultValue() const { return ArrayPtr->GetDefaultValue(); }
+	FORCEINLINE Type GetDefaultValue() const { return ArrayPtr->GetDefaultValue(); }
 
-	/** Return number of indices this attribute has */
+	UE_DEPRECATED(4.26, "Please use GetNumChannels().")
 	FORCEINLINE int32 GetNumIndices() const
 	{
-		return ArrayPtr->ArrayType::GetNumIndices();	// note: override virtual dispatch
+		return ArrayPtr->ArrayType::GetNumChannels();	// note: override virtual dispatch
+	}
+
+	/** Return number of indices this attribute has */
+	FORCEINLINE int32 GetNumChannels() const
+	{
+		return ArrayPtr->ArrayType::GetNumChannels();	// note: override virtual dispatch
 	}
 
 	/** Get the number of elements in this attribute array */
@@ -521,37 +650,69 @@ public:
 	/** Get the flags for this attribute array set */
 	FORCEINLINE EMeshAttributeFlags GetFlags() const { return ArrayPtr->GetFlags(); }
 
+	FORCEINLINE uint32 GetExtent() const { return 1; }
+
 	/** Set the element with the given ID and index 0 to the provided value */
-	FORCEINLINE void Set( const ElementIDType ElementID, const AttributeType& Value ) const
+	template <typename T = ElementIDType, typename TEnableIf<TIsDerivedFrom<T, FElementID>::Value, int>::Type = 0>
+	FORCEINLINE void Set(const ElementIDType ElementID, const Type& Value) const
+	{
+		ArrayPtr->GetArrayForChannel(0).GetElementBase(ElementID.GetValue())[0] = Value;
+	}
+
+	/** Set the element with the given ID and channel to the provided value */
+	template <typename T = ElementIDType, typename TEnableIf<TIsDerivedFrom<T, FElementID>::Value, int>::Type = 0>
+	FORCEINLINE void Set(const ElementIDType ElementID, const int32 Channel, const Type& Value) const
+	{
+		ArrayPtr->GetArrayForChannel(Channel).GetElementBase(ElementID.GetValue())[0] = Value;
+	}
+
+	FORCEINLINE void Set(int32 ElementIndex, const Type& Value) const
 	{
-		ArrayPtr->GetArrayForIndex( 0 )[ ElementID.GetValue() ] = Value;
+		ArrayPtr->GetArrayForChannel(0).GetElementBase(ElementIndex)[0] = Value;
 	}
 
-	/** Set the element with the given ID and index to the provided value */
-	FORCEINLINE void Set( const ElementIDType ElementID, const int32 Index, const AttributeType& Value ) const
+	FORCEINLINE void Set(int32 ElementIndex, const int32 Channel, const Type& Value) const
 	{
-		ArrayPtr->GetArrayForIndex( Index )[ ElementID.GetValue() ] = Value;
+		ArrayPtr->GetArrayForChannel(Channel).GetElementBase(ElementIndex)[0] = Value;
 	}
 
 	/** Copies the given attribute array and index to this index */
-	void Copy( TMeshAttributesRef<ElementIDType, const AttributeType> Src, const int32 DestIndex = 0, const int32 SrcIndex = 0 );
+	void Copy(TMeshAttributesRef<ElementIDType, const AttributeType, false> Src, const int32 DestIndex = 0, const int32 SrcIndex = 0);
 
-	/** Sets number of indices this attribute has */
-	FORCEINLINE void SetNumIndices( const int32 NumIndices ) const
+	UE_DEPRECATED(4.26, "Please use SetNumChannels().")
+	FORCEINLINE void SetNumIndices(const int32 NumChannels) const
 	{
-		ArrayPtr->ArrayType::SetNumIndices( NumIndices );	// note: override virtual dispatch
+		ArrayPtr->ArrayType::SetNumChannels(NumChannels);	// note: override virtual dispatch
 	}
 
-	/** Inserts an attribute index */
-	FORCEINLINE void InsertIndex( const int32 Index ) const
+	/** Sets number of channels this attribute has */
+	FORCEINLINE void SetNumChannels(const int32 NumChannels) const
 	{
-		ArrayPtr->ArrayType::InsertIndex( Index );		// note: override virtual dispatch
+		ArrayPtr->ArrayType::SetNumChannels(NumChannels);	// note: override virtual dispatch
 	}
 
-	/** Removes an attribute index */
-	FORCEINLINE void RemoveIndex( const int32 Index ) const
+	UE_DEPRECATED(4.26, "Please use InsertChannel().")
+	FORCEINLINE void InsertIndex(const int32 Index) const
 	{
-		ArrayPtr->ArrayType::RemoveIndex( Index );		// note: override virtual dispatch
+		ArrayPtr->ArrayType::InsertChannel(Index);		// note: override virtual dispatch
+	}
+
+	/** Inserts an attribute channel */
+	FORCEINLINE void InsertChannel(const int32 Index) const
+	{
+		ArrayPtr->ArrayType::InsertChannel(Index);		// note: override virtual dispatch
+	}
+
+	UE_DEPRECATED(4.26, "Please use RemoveChannel().")
+	FORCEINLINE void RemoveIndex(const int32 Index) const
+	{
+		ArrayPtr->ArrayType::RemoveChannel(Index);		// note: override virtual dispatch
+	}
+
+	/** Removes an attribute channel */
+	FORCEINLINE void RemoveChannel(const int32 Index) const
+	{
+		ArrayPtr->ArrayType::RemoveChannel(Index);		// note: override virtual dispatch
 	}
 
 protected:
@@ -559,18 +720,181 @@ protected:
 };
 
 template <typename ElementIDType, typename AttributeType>
-using TMeshAttributesConstRef = TMeshAttributesRef<ElementIDType, const AttributeType>;
+void TMeshAttributesRef<ElementIDType, AttributeType, false>::Copy(TMeshAttributesRef<ElementIDType, const AttributeType, false> Src, const int32 DestIndex, const int32 SrcIndex)
+{
+	check(Src.IsValid());
+	const TMeshAttributeArrayBase<AttributeType>& SrcArray = Src->ArrayPtr->GetArrayForChannel(SrcIndex);
+	TMeshAttributeArrayBase<AttributeType>& DestArray = ArrayPtr->GetArrayForChannel(DestIndex);
+	const int32 Num = FMath::Min(SrcArray.Num(), DestArray.Num());
+	for (int32 Index = 0; Index < Num; Index++)
+	{
+		DestArray.GetElementBase(Index)[0] = SrcArray.GetElementBase(Index)[0];
+	}
+}
+
 
 template <typename ElementIDType, typename AttributeType>
-void TMeshAttributesRef<ElementIDType, AttributeType>::Copy( TMeshAttributesRef<ElementIDType, const AttributeType> Src, const int32 DestIndex, const int32 SrcIndex )
+class TMeshAttributesRef<ElementIDType, AttributeType, true>
 {
-	check( Src.IsValid() );
-	const TMeshAttributeArrayBase<AttributeType>& SrcArray = Src->ArrayPtr->GetArrayForIndex( SrcIndex );
-	TMeshAttributeArrayBase<AttributeType>& DestArray = ArrayPtr->GetArrayForIndex( DestIndex );
-	const int32 Num = FMath::Min( SrcArray.Num(), DestArray.Num() );
-	for( int32 Index = 0; Index < Num; Index++ )
+	template <typename T, typename U, bool IsArray> friend class TMeshAttributesRef;
+
+public:
+	using Type = typename TBreakArrayView<AttributeType>::Type;
+	using ArrayType = typename TCopyQualifiersFromTo<Type, TMeshAttributeArraySet<typename TRemoveCV<Type>::Type>>::Type;
+
+	/** Constructor taking a pointer to a TMeshAttributeArraySet */
+	FORCEINLINE explicit TMeshAttributesRef(ArrayType* InArrayPtr = nullptr, uint32 InExtent = 1)
+		: ArrayPtr(InArrayPtr),
+		  Extent(InExtent)
+	{}
+
+	/** Implicitly construct a TMeshAttributesRef-to-const from a regular one */
+	template <typename T = Type, typename TEnableIf<TIsSame<T, const T>::Value, int>::Type = 0>
+	FORCEINLINE TMeshAttributesRef(TMeshAttributesRef<ElementIDType, TArrayView<typename TRemoveCV<T>::Type>> InRef)
+		: ArrayPtr(InRef.ArrayPtr),
+		  Extent(InRef.Extent)
+	{}
+
+	/** Access elements from attribute channel 0 */
+	template <typename T = ElementIDType, typename TEnableIf<TIsDerivedFrom<T, FElementID>::Value, int>::Type = 0>
+	FORCEINLINE TArrayView<Type> operator[](const ElementIDType ElementID) const
+	{
+		Type* Element = ArrayPtr->GetArrayForChannel(0).GetElementBase(ElementID.GetValue());
+		return TArrayView<Type>(Element, Extent);
+	}
+
+	/** Get the element with the given ID from channel 0 */
+	template <typename T = ElementIDType, typename TEnableIf<TIsDerivedFrom<T, FElementID>::Value, int>::Type = 0>
+	FORCEINLINE TArrayView<Type> Get(const ElementIDType ElementID) const
+	{
+		Type* Element = ArrayPtr->GetArrayForChannel(0).GetElementBase(ElementID.GetValue());
+		return TArrayView<Type>(Element, Extent);
+	}
+
+	/** Get the element with the given ID and channel */
+	template <typename T = ElementIDType, typename TEnableIf<TIsDerivedFrom<T, FElementID>::Value, int>::Type = 0>
+	FORCEINLINE TArrayView<Type> Get(const ElementIDType ElementID, const int32 Channel) const
+	{
+		Type* Element = ArrayPtr->GetArrayForChannel(Channel).GetElementBase(ElementID.GetValue());
+		return TArrayView<Type>(Element, Extent);
+	}
+
+	FORCEINLINE TArrayView<Type> operator[](int32 ElementIndex) const
+	{
+		Type* Element = ArrayPtr->GetArrayForChannel(0).GetElementBase(ElementIndex);
+		return TArrayView<Type>(Element, Extent);
+	}
+
+	FORCEINLINE TArrayView<Type> Get(int32 ElementIndex) const
+	{
+		Type* Element = ArrayPtr->GetArrayForChannel(0).GetElementBase(ElementIndex);
+		return TArrayView<Type>(Element, Extent);
+	}
+
+	FORCEINLINE TArrayView<Type> Get(int32 ElementIndex, const int32 Channel) const
+	{
+		Type* Element = ArrayPtr->GetArrayForChannel(Channel).GetElementBase(ElementIndex);
+		return TArrayView<Type>(Element, Extent);
+	}
+
+	FORCEINLINE TArrayView<Type> GetRawArray(const int32 ChannelIndex = 0) const
+	{
+		if (ArrayPtr == nullptr)
+		{
+			return TArrayView<Type>();
+		}
+
+		Type* Element = ArrayPtr->GetArrayForChannel(ChannelIndex).GetElementBase(0);
+		return TArrayView<Type>(Element, GetNumElements() * Extent);
+	}
+
+	/** Return whether the reference is valid or not */
+	FORCEINLINE bool IsValid() const { return (ArrayPtr != nullptr); }
+
+	/** Return default value for this attribute type */
+	FORCEINLINE Type GetDefaultValue() const { return ArrayPtr->GetDefaultValue(); }
+
+	UE_DEPRECATED(4.26, "Please use GetNumChannels().")
+	FORCEINLINE int32 GetNumIndices() const
+	{
+		return ArrayPtr->ArrayType::GetNumChannels();	// note: override virtual dispatch
+	}
+
+	/** Return number of channels this attribute has */
+	FORCEINLINE int32 GetNumChannels() const
+	{
+		return ArrayPtr->ArrayType::GetNumChannels();	// note: override virtual dispatch
+	}
+
+	/** Get the number of elements in this attribute array */
+	FORCEINLINE int32 GetNumElements() const
+	{
+		return ArrayPtr->GetNumElements();
+	}
+
+	/** Get the flags for this attribute array set */
+	FORCEINLINE EMeshAttributeFlags GetFlags() const { return ArrayPtr->GetFlags(); }
+
+	FORCEINLINE uint32 GetExtent() const { return Extent; }
+
+	/** Copies the given attribute array and index to this index */
+	void Copy(TMeshAttributesConstRef<ElementIDType, AttributeType, true> Src, const int32 DestIndex = 0, const int32 SrcIndex = 0);
+
+	UE_DEPRECATED(4.26, "Please use SetNumChannels().")
+	FORCEINLINE void SetNumIndices(const int32 NumChannels) const
+	{
+		ArrayPtr->ArrayType::SetNumChannels(NumChannels);	// note: override virtual dispatch
+	}
+
+	/** Sets number of channels this attribute has */
+	FORCEINLINE void SetNumChannels(const int32 NumChannels) const
 	{
-		DestArray[ Index ] = SrcArray[ Index ];
+		ArrayPtr->ArrayType::SetNumChannels(NumChannels);	// note: override virtual dispatch
+	}
+
+	UE_DEPRECATED(4.26, "Please use InsertChannel().")
+	FORCEINLINE void InsertIndex(const int32 Index) const
+	{
+		ArrayPtr->ArrayType::InsertChannel(Index);		// note: override virtual dispatch
+	}
+
+	/** Inserts an attribute channel */
+	FORCEINLINE void InsertChannel(const int32 Index) const
+	{
+		ArrayPtr->ArrayType::InsertChannel(Index);		// note: override virtual dispatch
+	}
+
+	UE_DEPRECATED(4.26, "Please use RemoveChannel().")
+	FORCEINLINE void RemoveIndex(const int32 Index) const
+	{
+		ArrayPtr->ArrayType::RemoveChannel(Index);		// note: override virtual dispatch
+	}
+
+	/** Removes an attribute channel */
+	FORCEINLINE void RemoveChannel(const int32 Index) const
+	{
+		ArrayPtr->ArrayType::RemoveChannel(Index);		// note: override virtual dispatch
+	}
+
+protected:
+	ArrayType* ArrayPtr;
+	uint32 Extent;
+};
+
+template <typename ElementIDType, typename AttributeType>
+void TMeshAttributesRef<ElementIDType, AttributeType, true>::Copy(TMeshAttributesConstRef<ElementIDType, AttributeType, true> Src, const int32 DestIndex, const int32 SrcIndex)
+{
+	check(Src.IsValid());
+	check(Src.Extent == Extent);
+	const TMeshAttributeArrayBase<Type>& SrcArray = Src->ArrayPtr->GetArrayForChannel(SrcIndex);
+	TMeshAttributeArrayBase<Type>& DestArray = ArrayPtr->GetArrayForChannel(DestIndex);
+	const int32 Num = FMath::Min(SrcArray.Num(), DestArray.Num());
+	for (int32 Index = 0; Index < Num; Index++)
+	{
+		for (uint32 Count = 0; Count < Extent; Count++)
+		{
+			DestArray.GetElementBase(Index)[Count] = SrcArray.GetElementBase(Index)[Count];
+		}
 	}
 }
 
@@ -660,7 +984,7 @@ public:
 	FORCEINLINE void Set( const ElementIDType ElementID, const int32 Index, const ViewType& Value ) const { this->SetByIndex( ElementID.GetValue(), Index, Value ); }
 
 	/** Sets number of indices this attribute has */
-	FORCEINLINE void SetNumIndices( const int32 NumIndices ) const { this->ArrayPtr->SetNumIndices( NumIndices ); }
+	FORCEINLINE void SetNumIndices( const int32 NumIndices ) const { this->ArrayPtr->SetNumChannels( NumIndices ); }
 
 	/** Inserts an attribute index */
 	FORCEINLINE void InsertIndex( const int32 Index ) const { this->ArrayPtr->InsertIndex( Index ); }
@@ -691,31 +1015,31 @@ public:
 	 * Construct a valid FAttributesSetEntry of the concrete type specified.
 	 */
 	template <typename AttributeType>
-	FORCEINLINE FAttributesSetEntry( const int32 NumberOfIndices, const AttributeType& Default, const EMeshAttributeFlags Flags, const int32 NumElements )
-		: Ptr( MakeUnique<TMeshAttributeArraySet<AttributeType>>( NumberOfIndices, Default, Flags, NumElements ) )
+	FORCEINLINE FAttributesSetEntry(const int32 NumberOfChannels, const AttributeType& Default, const EMeshAttributeFlags Flags, const int32 NumElements, const int32 Extent)
+		: Ptr(MakeUnique<TMeshAttributeArraySet<AttributeType>>(NumberOfChannels, Default, Flags, NumElements, Extent))
 	{}
 
 	/** Default destructor */
 	FORCEINLINE ~FAttributesSetEntry() = default;
 
 	/** Polymorphic copy: a new copy of Other is created */
-	FAttributesSetEntry( const FAttributesSetEntry& Other )
-		: Ptr( Other.Ptr ? Other.Ptr->Clone() : nullptr )
+	FAttributesSetEntry(const FAttributesSetEntry& Other)
+		: Ptr(Other.Ptr ? Other.Ptr->Clone() : nullptr)
 	{}
 
 	/** Default move constructor */
-	FAttributesSetEntry( FAttributesSetEntry&& ) = default;
+	FAttributesSetEntry(FAttributesSetEntry&&) = default;
 
 	/** Polymorphic assignment */
-	FAttributesSetEntry& operator=( const FAttributesSetEntry& Other )
+	FAttributesSetEntry& operator=(const FAttributesSetEntry& Other)
 	{
-		FAttributesSetEntry Temp( Other );
-		Swap( *this, Temp );
+		FAttributesSetEntry Temp(Other);
+		Swap(*this, Temp);
 		return *this;
 	}
 
 	/** Default move assignment */
-	FAttributesSetEntry& operator=( FAttributesSetEntry&& ) = default;
+	FAttributesSetEntry& operator=(FAttributesSetEntry&&) = default;
 
 	/** Transparent access through the TUniquePtr */
 	FORCEINLINE const FMeshAttributeArraySetBase* Get() const { return Ptr.Get(); }
@@ -730,10 +1054,10 @@ public:
 	FORCEINLINE bool operator!() const { return !Ptr.IsValid(); }
 
 	/** Given a type at runtime, allocate an attribute array of that type, owned by Ptr */
-	void CreateArrayOfType( const uint32 Type );
+	void CreateArrayOfType(const uint32 Type, const uint32 Extent);
 
 	/** Serialization */
-	friend FArchive& operator<<( FArchive& Ar, FAttributesSetEntry& Entry );
+	friend FArchive& operator<<(FArchive& Ar, FAttributesSetEntry& Entry);
 
 private:
 	TUniquePtr<FMeshAttributeArraySetBase> Ptr;
@@ -749,12 +1073,12 @@ class FAttributesSetBase
 public:
 	/** Constructor */
 	FAttributesSetBase()
-		: NumElements( 0 )
+		: NumElements(0)
 	{}
 
 	/**
 	 * Register a new attribute name with the given type (must be a member of the AttributeTypes tuple).
-	 * If the attribute name is already registered, it will update it to use the new type, number of indices and flags.
+	 * If the attribute name is already registered, it will update it to use the new type, number of channels and flags.
 	 *
 	 * Example of use:
 	 *
@@ -763,93 +1087,120 @@ public:
 	 *		TVertexInstanceAttributeArray<FVector2D>& UV0 = VertexInstanceAttributes().GetAttributes<FVector2D>( "UV", 0 );
 	 *		UV0[ VertexInstanceID ] = FVector2D( 1.0f, 1.0f );
 	 */
-	template <typename AttributeType>
-	void RegisterAttribute( const FName AttributeName, const int32 NumberOfIndices = 1, const AttributeType& Default = AttributeType(), const EMeshAttributeFlags Flags = EMeshAttributeFlags::None );
+	template <typename AttributeType,
+			  typename TEnableIf<!TIsBoundedArray<AttributeType>::Value, int>::Type = 0>
+	TMeshAttributesArray<AttributeType> RegisterAttribute(const FName AttributeName, const int32 NumberOfChannels = 1, const AttributeType& Default = AttributeType(), const EMeshAttributeFlags Flags = EMeshAttributeFlags::None );
+
+	template <typename ArrayType,
+			  typename TEnableIf<TIsBoundedArray<ArrayType>::Value, int>::Type = 0,
+			  typename AttributeType = typename TBreakArray<ArrayType>::Type,
+			  uint32 Extent = TBreakArray<ArrayType>::Value>
+	TMeshAttributesArray<TArrayView<AttributeType>> RegisterAttribute(const FName AttributeName, const int32 NumberOfChannels = 1, const typename TBreakArray<ArrayType>::Type& Default = AttributeType(), const EMeshAttributeFlags Flags = EMeshAttributeFlags::None);
+
+	/**
+	 * Register a new index attribute (type is implicitly int).
+	 *
+	 * If the attribute name is already registered, it will update it to use the new type, number of channels and flags.
+	 */
+	template <typename AttributeType,
+			  typename TEnableIf<TIsSame<AttributeType, int>::Value || TIsDerivedFrom<AttributeType, FElementID>::Value, int>::Type = 0>
+	TMeshAttributesArray<AttributeType> RegisterIndexAttribute(const FName AttributeName, const int32 NumberOfChannels = 1, const EMeshAttributeFlags Flags = EMeshAttributeFlags::IndexReference);
+
+	/**
+	 * Register a new attribute denoting an array of indices (type is implicitly int).
+	 *
+	 * If the attribute name is already registered, it will update it to use the new type, number of channels and flags.
+	 */
+	template <typename ArrayType,
+			  typename AttributeType = typename TBreakArray<ArrayType>::Type,
+			  typename TEnableIf<TIsSame<AttributeType, int>::Value || TIsDerivedFrom<AttributeType, FElementID>::Value, int>::Type = 0,
+			  uint32 Extent = TBreakArray<ArrayType>::Value>
+	TMeshAttributesArray<TArrayView<AttributeType>> RegisterIndexAttribute(const FName AttributeName, const int32 NumberOfChannels = 1, const EMeshAttributeFlags Flags = EMeshAttributeFlags::IndexReference);
 
 	/**
 	 * Unregister an attribute with the given name.
 	 */
-	void UnregisterAttribute( const FName AttributeName )
+	void UnregisterAttribute(const FName AttributeName)
 	{
-		Map.Remove( AttributeName );
+		Map.Remove(AttributeName);
 	}
 
 	/** Determines whether an attribute exists with the given name */
-	bool HasAttribute( const FName AttributeName ) const
+	bool HasAttribute(const FName AttributeName) const
 	{
-		return ( Map.Contains( AttributeName ) );
+		return Map.Contains(AttributeName);
 	}
 
 	/**
 	 * Determines whether an attribute of the given type exists with the given name
 	 */
-	template <typename AttributeType>
-	bool HasAttributeOfType( const FName AttributeName ) const
+	template <typename AttributeType,
+			  typename TEnableIf<!TIsBoundedArray<AttributeType>::Value, int>::Type = 0>
+	bool HasAttributeOfType(const FName AttributeName) const
+	{
+		if (const FAttributesSetEntry* ArraySetPtr = Map.Find(AttributeName))
+		{
+			return (*ArraySetPtr)->HasType<AttributeType>() && (*ArraySetPtr)->GetExtent() == 1;
+		}
+
+		return false;
+	}
+
+	template <typename ArrayType,
+			  typename TEnableIf<TIsBoundedArray<ArrayType>::Value, int>::Type = 0,
+			  typename AttributeType = typename TBreakArray<ArrayType>::Type,
+			  uint32 Extent = TBreakArray<ArrayType>::Value>
+	bool HasAttributeOfType(const FName AttributeName) const
 	{
-		if( const FAttributesSetEntry* ArraySetPtr = Map.Find( AttributeName ) )
+		if (const FAttributesSetEntry* ArraySetPtr = Map.Find(AttributeName))
 		{
-			return ( *ArraySetPtr )->HasType<AttributeType>();
+			return (*ArraySetPtr)->HasType<AttributeType>() && (*ArraySetPtr)->GetExtent() == Extent;
 		}
 
 		return false;
 	}
 
 	/** Initializes all attributes to have the given number of elements with the default value */
-	void Initialize( const int32 Count )
+	void Initialize(const int32 Count)
 	{
 		NumElements = Count;
-		for( auto& MapEntry : Map )
+		for (auto& MapEntry : Map)
 		{
-			MapEntry.Value->Initialize( Count );
+			MapEntry.Value->Initialize(Count);
 		}
 	}
 
 	/** Applies the given remapping to the attributes set */
-	void Remap( const TSparseArray<int32>& IndexRemap );
+	void Remap(const TSparseArray<int32>& IndexRemap);
 
 	/** Returns an array of all the attribute names registered */
 	template <typename Allocator>
-	void GetAttributeNames( TArray<FName, Allocator>& OutAttributeNames ) const
+	void GetAttributeNames(TArray<FName, Allocator>& OutAttributeNames) const
 	{
-		Map.GetKeys( OutAttributeNames );
+		Map.GetKeys(OutAttributeNames);
 	}
 
-	uint32 GetHash( const FName AttributeName ) const
+	uint32 GetHash(const FName AttributeName) const
 	{
-		if ( const FAttributesSetEntry* ArraySetPtr = Map.Find( AttributeName ) )
+		if (const FAttributesSetEntry* ArraySetPtr = Map.Find(AttributeName))
 		{
-			return ( *ArraySetPtr )->GetHash();
+			return (*ArraySetPtr)->GetHash();
 		}
 		return 0;
 	}
 
-	template <typename AttributeType>
-	UE_DEPRECATED( 4.20, "Please use untemplated UnregisterAttribute() instead" )
-	void UnregisterAttribute( const FName AttributeName )
-	{
-		return UnregisterAttribute( AttributeName );
-	}
-
-	template <typename AttributeType>
-	UE_DEPRECATED( 4.20, "Please use untemplated HasAttribute() instead" )
-	bool HasAttribute( const FName AttributeName )
-	{
-		return HasAttribute( AttributeName );
-	}
-
-protected:
 	/**
 	 * Insert a new element at the given index.
 	 * The public API version of this function takes an ID of ElementIDType instead of a typeless index.
 	 */
-	void Insert( const int32 Index )
+	void Insert(const int32 Index)
 	{
-		NumElements = FMath::Max( NumElements, Index + 1 );
+		NumElements = FMath::Max(NumElements, Index + 1);
 
-		for( auto& MapEntry : Map )
+		for (auto& MapEntry : Map)
 		{
-			MapEntry.Value->Insert( Index );
-			check( MapEntry.Value->GetNumElements() == NumElements );
+			MapEntry.Value->Insert(Index);
+			check(MapEntry.Value->GetNumElements() == NumElements);
 		}
 	}
 
@@ -857,19 +1208,92 @@ protected:
 	 * Remove an element at the given index.
 	 * The public API version of this function takes an ID of ElementIDType instead of a typeless index.
 	 */
-	void Remove( const int32 Index )
+	void Remove(const int32 Index)
+	{
+		for (auto& MapEntry : Map)
+		{
+			MapEntry.Value->Remove(Index);
+		}
+	}
+
+	/**
+	 * Get an attribute array with the given type and name.
+	 * The attribute type must correspond to the type passed as the template parameter.
+	 */
+	template <typename AttributeType,
+			  typename TEnableIf<!TIsArrayView<AttributeType>::Value, int>::Type = 0>
+	TMeshAttributesConstRef<int32, AttributeType> GetAttributesRef(const FName AttributeName) const
+	{
+		if (const FAttributesSetEntry* ArraySetPtr = this->Map.Find(AttributeName))
+		{
+			using Type = typename TChooseClass<TIsDerivedFrom<AttributeType, FElementID>::Value, int, AttributeType>::Result;
+			if ((*ArraySetPtr)->HasType<Type>() && (*ArraySetPtr)->GetExtent() == 1)
+			{
+				return TMeshAttributesConstRef<int32, AttributeType>(static_cast<const TMeshAttributeArraySet<AttributeType>*>(ArraySetPtr->Get()));
+			}
+		}
+
+		return TMeshAttributesConstRef<int32, AttributeType>();
+	}
+
+	template <typename AttributeType,
+			  typename TEnableIf<!TIsArrayView<AttributeType>::Value, int>::Type = 0>
+	TMeshAttributesRef<int32, AttributeType> GetAttributesRef(const FName AttributeName)
 	{
-		for( auto& MapEntry : Map )
+		if (FAttributesSetEntry* ArraySetPtr = this->Map.Find(AttributeName))
 		{
-			MapEntry.Value->Remove( Index );
+			using Type = typename TChooseClass<TIsDerivedFrom<AttributeType, FElementID>::Value, int, AttributeType>::Result;
+			if ((*ArraySetPtr)->HasType<Type>() && (*ArraySetPtr)->GetExtent() == 1)
+			{
+				return TMeshAttributesRef<int32, AttributeType>(static_cast<TMeshAttributeArraySet<AttributeType>*>(ArraySetPtr->Get()));
+			}
 		}
+
+		return TMeshAttributesRef<int32, AttributeType>();
 	}
 
+	template <typename ArrayType,
+			  typename TEnableIf<TIsArrayView<ArrayType>::Value, int>::Type = 0,
+			  typename AttributeType = typename TBreakArrayView<ArrayType>::Type>
+	TMeshAttributesConstRef<int32, ArrayType> GetAttributesRef(const FName AttributeName) const
+	{
+		if (const FAttributesSetEntry* ArraySetPtr = this->Map.Find(AttributeName))
+		{
+			using Type = typename TChooseClass<TIsDerivedFrom<AttributeType, FElementID>::Value, int, AttributeType>::Result;
+			if ((*ArraySetPtr)->HasType<Type>())
+			{
+				return TMeshAttributesConstRef<int32, ArrayType>(static_cast<const TMeshAttributeArraySet<AttributeType>*>(ArraySetPtr->Get()), (*ArraySetPtr)->GetExtent());
+			}
+		}
+
+		return TMeshAttributesConstRef<int32, ArrayType>();
+	}
+
+	template <typename ArrayType,
+			  typename TEnableIf<TIsArrayView<ArrayType>::Value, int>::Type = 0,
+			  typename AttributeType = typename TBreakArrayView<ArrayType>::Type>
+	TMeshAttributesRef<int32, ArrayType> GetAttributesRef(const FName AttributeName)
+	{
+		if (FAttributesSetEntry* ArraySetPtr = this->Map.Find(AttributeName))
+		{
+			using Type = typename TChooseClass<TIsDerivedFrom<AttributeType, FElementID>::Value, int, AttributeType>::Result;
+			if ((*ArraySetPtr)->HasType<Type>())
+			{
+				return TMeshAttributesRef<int32, ArrayType>(static_cast<TMeshAttributeArraySet<AttributeType>*>(ArraySetPtr->Get()), (*ArraySetPtr)->GetExtent());
+			}
+		}
+
+		return TMeshAttributesRef<int32, ArrayType>();
+	}
+
+	void AppendAttributesFrom(const FAttributesSetBase& OtherAttributesSet);
+
+protected:
 	/** Serialization */
-	friend MESHDESCRIPTION_API FArchive& operator<<( FArchive& Ar, FAttributesSetBase& AttributesSet );
+	friend MESHDESCRIPTION_API FArchive& operator<<(FArchive& Ar, FAttributesSetBase& AttributesSet);
 
 	template <typename T>
-	friend void SerializeLegacy( FArchive& Ar, FAttributesSetBase& AttributesSet );
+	friend void SerializeLegacy(FArchive& Ar, FAttributesSetBase& AttributesSet);
 
 	/** The actual container */
 	TMap<FName, FAttributesSetEntry> Map;
@@ -879,26 +1303,119 @@ protected:
 };
 
 
-template <typename AttributeType>
-void FAttributesSetBase::RegisterAttribute( const FName AttributeName, const int32 NumberOfIndices, const AttributeType& Default, const EMeshAttributeFlags Flags )
+template <typename AttributeType,
+		  typename TEnableIf<!TIsBoundedArray<AttributeType>::Value, int>::Type>
+TMeshAttributesArray<AttributeType> FAttributesSetBase::RegisterAttribute(const FName AttributeName, const int32 NumberOfChannels, const AttributeType& Default, const EMeshAttributeFlags Flags)
+{
+	using ArrayType = TMeshAttributeArraySet<AttributeType>;
+	if (FAttributesSetEntry* ArraySetPtr = Map.Find(AttributeName))
+	{
+		if ((*ArraySetPtr)->HasType<AttributeType>() && (*ArraySetPtr)->GetExtent() == 1)
+		{
+			static_cast<ArrayType*>(ArraySetPtr->Get())->ArrayType::SetNumChannels(NumberOfChannels);	// note: override virtual dispatch
+			(*ArraySetPtr)->SetFlags(Flags);
+			return TMeshAttributesArray<AttributeType>(static_cast<ArrayType*>(ArraySetPtr->Get()));
+		}
+		else
+		{
+			Map.Remove(AttributeName);
+			FAttributesSetEntry& Entry = Map.Emplace(AttributeName, FAttributesSetEntry(NumberOfChannels, Default, Flags, NumElements, 1));
+			return TMeshAttributesArray<AttributeType>(static_cast<ArrayType*>(Entry.Get()));
+		}
+	}
+	else
+	{
+		FAttributesSetEntry& Entry = Map.Emplace(AttributeName, FAttributesSetEntry(NumberOfChannels, Default, Flags, NumElements, 1));
+		return TMeshAttributesArray<AttributeType>(static_cast<ArrayType*>(Entry.Get()));
+	}
+}
+
+
+template <typename ArrayType,
+		  typename TEnableIf<TIsBoundedArray<ArrayType>::Value, int>::Type,
+		  typename AttributeType,
+		  uint32 Extent>
+TMeshAttributesArray<TArrayView<AttributeType>> FAttributesSetBase::RegisterAttribute(const FName AttributeName, const int32 NumberOfChannels, const typename TBreakArray<ArrayType>::Type& Default, const EMeshAttributeFlags Flags)
 {
-	if( FAttributesSetEntry* ArraySetPtr = Map.Find( AttributeName ) )
+	using ContainerType = TMeshAttributeArraySet<AttributeType>;
+	if (FAttributesSetEntry* ArraySetPtr = Map.Find(AttributeName))
 	{
-		if( ( *ArraySetPtr )->HasType<AttributeType>() )
+		if ((*ArraySetPtr)->HasType<AttributeType>() && (*ArraySetPtr)->GetExtent() == Extent)
 		{
-			using ArrayType = TMeshAttributeArraySet<AttributeType>;
-			static_cast<ArrayType*>( ArraySetPtr->Get() )->ArrayType::SetNumIndices( NumberOfIndices );	// note: override virtual dispatch
-			( *ArraySetPtr )->SetFlags( Flags );
+			static_cast<ContainerType*>(ArraySetPtr->Get())->ContainerType::SetNumChannels(NumberOfChannels);	// note: override virtual dispatch
+			(*ArraySetPtr)->SetFlags(Flags);
+			return TMeshAttributesArray<TArrayView<AttributeType>>(static_cast<ContainerType*>(ArraySetPtr->Get()), Extent);
 		}
 		else
 		{
-			Map.Remove( AttributeName );
-			Map.Emplace( AttributeName, FAttributesSetEntry( NumberOfIndices, Default, Flags, NumElements ) );
+			Map.Remove(AttributeName);
+			FAttributesSetEntry& Entry = Map.Emplace(AttributeName, FAttributesSetEntry(NumberOfChannels, Default, Flags, NumElements, Extent));
+			return TMeshAttributesArray<TArrayView<AttributeType>>(static_cast<ContainerType*>(Entry.Get()), Extent);
 		}
 	}
 	else
 	{
-		Map.Emplace( AttributeName, FAttributesSetEntry( NumberOfIndices, Default, Flags, NumElements ) );
+		FAttributesSetEntry& Entry = Map.Emplace(AttributeName, FAttributesSetEntry(NumberOfChannels, Default, Flags, NumElements, Extent));
+		return TMeshAttributesArray<TArrayView<AttributeType>>(static_cast<ContainerType*>(Entry.Get()), Extent);
+	}
+}
+
+
+template <typename AttributeType,
+		  typename TEnableIf<TIsSame<AttributeType, int>::Value || TIsDerivedFrom<AttributeType, FElementID>::Value, int>::Type>
+TMeshAttributesArray<AttributeType> FAttributesSetBase::RegisterIndexAttribute(const FName AttributeName, const int32 NumberOfChannels, const EMeshAttributeFlags Flags)
+{
+	using ArrayType = TMeshAttributeArraySet<AttributeType>;
+	if (FAttributesSetEntry* ArraySetPtr = Map.Find(AttributeName))
+	{
+		if ((*ArraySetPtr)->HasType<int>() && (*ArraySetPtr)->GetExtent() == 1)
+		{
+			static_cast<ArrayType*>(ArraySetPtr->Get())->ArrayType::SetNumChannels(NumberOfChannels);	// note: override virtual dispatch
+			(*ArraySetPtr)->SetFlags(Flags);
+			return TMeshAttributesArray<AttributeType>(static_cast<ArrayType*>(ArraySetPtr->Get()));
+		}
+		else
+		{
+			Map.Remove(AttributeName);
+			FAttributesSetEntry& Entry = Map.Emplace(AttributeName, FAttributesSetEntry(NumberOfChannels, int32(INDEX_NONE), Flags, NumElements, 1));
+			return TMeshAttributesArray<AttributeType>(static_cast<ArrayType*>(Entry.Get()));
+		}
+	}
+	else
+	{
+		FAttributesSetEntry& Entry = Map.Emplace(AttributeName, FAttributesSetEntry(NumberOfChannels, int32(INDEX_NONE), Flags, NumElements, 1));
+		return TMeshAttributesArray<AttributeType>(static_cast<ArrayType*>(Entry.Get()));
+	}
+
+}
+
+
+template <typename ArrayType,
+		  typename AttributeType,
+		  typename TEnableIf<TIsSame<AttributeType, int>::Value || TIsDerivedFrom<AttributeType, FElementID>::Value, int>::Type,
+		  uint32 Extent>
+TMeshAttributesArray<TArrayView<AttributeType>> FAttributesSetBase::RegisterIndexAttribute(const FName AttributeName, const int32 NumberOfChannels, const EMeshAttributeFlags Flags)
+{
+	if (FAttributesSetEntry* ArraySetPtr = Map.Find(AttributeName))
+	{
+		if ((*ArraySetPtr)->HasType<int>() && (*ArraySetPtr)->GetExtent() == Extent)
+		{
+			using ContainerType = TMeshAttributeArraySet<int>;
+			static_cast<ContainerType*>(ArraySetPtr->Get())->ContainerType::SetNumChannels(NumberOfChannels);	// note: override virtual dispatch
+			(*ArraySetPtr)->SetFlags(Flags);
+			return TMeshAttributesArray<TArrayView<AttributeType>>(static_cast<TMeshAttributeArraySet<AttributeType>*>(ArraySetPtr->Get()), Extent);
+		}
+		else
+		{
+			Map.Remove(AttributeName);
+			FAttributesSetEntry& Entry = Map.Emplace(AttributeName, FAttributesSetEntry(NumberOfChannels, int32(INDEX_NONE), Flags, NumElements, Extent));
+			return TMeshAttributesArray<TArrayView<AttributeType>>(static_cast<TMeshAttributeArraySet<AttributeType>*>(Entry.Get()), Extent);
+		}
+	}
+	else
+	{
+		FAttributesSetEntry& Entry = Map.Emplace(AttributeName, FAttributesSetEntry(NumberOfChannels, int32(INDEX_NONE), Flags, NumElements, Extent));
+		return TMeshAttributesArray<TArrayView<AttributeType>>(static_cast<TMeshAttributeArraySet<AttributeType>*>(Entry.Get()), Extent);
 	}
 }
 
@@ -930,34 +1447,72 @@ public:
 	 * Note that the returned object is a value type which should be assigned and passed by value, not reference.
 	 * It is valid for as long as this TAttributesSet object exists.
 	 */
-	template <typename AttributeType>
-	TMeshAttributesConstRef<ElementIDType, AttributeType> GetAttributesRef( const FName AttributeName ) const
+	template <typename AttributeType,
+			  typename TEnableIf<!TIsArrayView<AttributeType>::Value, int>::Type = 0>
+	TMeshAttributesConstRef<ElementIDType, AttributeType> GetAttributesRef(const FName AttributeName) const
 	{
-		if( const FAttributesSetEntry* ArraySetPtr = this->Map.Find( AttributeName ) )
+		if (const FAttributesSetEntry* ArraySetPtr = this->Map.Find(AttributeName))
 		{
-			if( ( *ArraySetPtr )->HasType<AttributeType>() )
+			using Type = typename TChooseClass<TIsDerivedFrom<AttributeType, FElementID>::Value, int, AttributeType>::Result;
+			if ((*ArraySetPtr)->HasType<Type>() && (*ArraySetPtr)->GetExtent() == 1)
 			{
-				return TMeshAttributesConstRef<ElementIDType, AttributeType>( static_cast<const TMeshAttributeArraySet<AttributeType>*>( ArraySetPtr->Get() ) );
+				return TMeshAttributesConstRef<ElementIDType, AttributeType>(static_cast<const TMeshAttributeArraySet<AttributeType>*>(ArraySetPtr->Get()));
 			}
 		}
 
 		return TMeshAttributesConstRef<ElementIDType, AttributeType>();
 	}
 
-	template <typename AttributeType>
-	TMeshAttributesRef<ElementIDType, AttributeType> GetAttributesRef( const FName AttributeName )
+	template <typename AttributeType,
+			  typename TEnableIf<!TIsArrayView<AttributeType>::Value, int>::Type = 0>
+	TMeshAttributesRef<ElementIDType, AttributeType> GetAttributesRef(const FName AttributeName)
 	{
-		if( FAttributesSetEntry* ArraySetPtr = this->Map.Find( AttributeName ) )
+		if (FAttributesSetEntry* ArraySetPtr = this->Map.Find(AttributeName))
 		{
-			if( ( *ArraySetPtr )->HasType<AttributeType>() )
+			using Type = typename TChooseClass<TIsDerivedFrom<AttributeType, FElementID>::Value, int, AttributeType>::Result;
+			if ((*ArraySetPtr)->HasType<Type>() && (*ArraySetPtr)->GetExtent() == 1)
 			{
-				return TMeshAttributesRef<ElementIDType, AttributeType>( static_cast<TMeshAttributeArraySet<AttributeType>*>( ArraySetPtr->Get() ) );
+				return TMeshAttributesRef<ElementIDType, AttributeType>(static_cast<TMeshAttributeArraySet<AttributeType>*>(ArraySetPtr->Get()));
 			}
 		}
 
 		return TMeshAttributesRef<ElementIDType, AttributeType>();
 	}
 
+	template <typename ArrayType,
+			  typename TEnableIf<TIsArrayView<ArrayType>::Value, int>::Type = 0,
+			  typename AttributeType = typename TBreakArrayView<ArrayType>::Type>
+	TMeshAttributesConstRef<ElementIDType, ArrayType> GetAttributesRef(const FName AttributeName) const
+	{
+		if (const FAttributesSetEntry* ArraySetPtr = this->Map.Find(AttributeName))
+		{
+			using Type = typename TChooseClass<TIsDerivedFrom<AttributeType, FElementID>::Value, int, AttributeType>::Result;
+			if ((*ArraySetPtr)->HasType<Type>())
+			{
+				return TMeshAttributesConstRef<ElementIDType, ArrayType>(static_cast<const TMeshAttributeArraySet<AttributeType>*>(ArraySetPtr->Get()), (*ArraySetPtr)->GetExtent());
+			}
+		}
+
+		return TMeshAttributesConstRef<ElementIDType, ArrayType>();
+	}
+
+	template <typename ArrayType,
+			  typename TEnableIf<TIsArrayView<ArrayType>::Value, int>::Type = 0,
+			  typename AttributeType = typename TBreakArrayView<ArrayType>::Type>
+	TMeshAttributesRef<ElementIDType, ArrayType> GetAttributesRef(const FName AttributeName)
+	{
+		if (FAttributesSetEntry* ArraySetPtr = this->Map.Find(AttributeName))
+		{
+			using Type = typename TChooseClass<TIsDerivedFrom<AttributeType, FElementID>::Value, int, AttributeType>::Result;
+			if ((*ArraySetPtr)->HasType<Type>())
+			{
+				return TMeshAttributesRef<ElementIDType, ArrayType>(static_cast<TMeshAttributeArraySet<AttributeType>*>(ArraySetPtr->Get()), (*ArraySetPtr)->GetExtent());
+			}
+		}
+
+		return TMeshAttributesRef<ElementIDType, ArrayType>();
+	}
+
 	/**
 	 * Get a view on an attribute array with the given name, accessing elements as the given type.
 	 * Access to elements will be slightly slower than with GetAttributesRef, but element access is not strongly typed.
@@ -976,178 +1531,230 @@ public:
 	 * It is valid for as long as this TAttributesSet object exists.
 	 */
 	template <typename ViewType>
-	TMeshAttributesConstView<ElementIDType, ViewType> GetAttributesView( const FName AttributeName ) const
+	UE_DEPRECATED(4.25, "Views are due to be deprecated. Please use MeshAttributeRefs instead.")
+	TMeshAttributesConstView<ElementIDType, ViewType> GetAttributesView(const FName AttributeName) const
 	{
-		if( const FAttributesSetEntry* ArraySetPtr = this->Map.Find( AttributeName ) )
+		if (const FAttributesSetEntry* ArraySetPtr = this->Map.Find(AttributeName))
 		{
-			return TMeshAttributesConstView<ElementIDType, ViewType>( ArraySetPtr->Get() );
+			return TMeshAttributesConstView<ElementIDType, ViewType>(ArraySetPtr->Get());
 		}
 
 		return TMeshAttributesConstView<ElementIDType, ViewType>();
 	}
 
 	template <typename ViewType>
-	TMeshAttributesView<ElementIDType, ViewType> GetAttributesView( const FName AttributeName )
+	UE_DEPRECATED(4.25, "Views are due to be deprecated. Please use MeshAttributeRefs instead.")
+	TMeshAttributesView<ElementIDType, ViewType> GetAttributesView(const FName AttributeName)
 	{
-		if( FAttributesSetEntry* ArraySetPtr = this->Map.Find( AttributeName ) )
+		if (FAttributesSetEntry* ArraySetPtr = this->Map.Find(AttributeName))
 		{
-			return TMeshAttributesView<ElementIDType, ViewType>( ArraySetPtr->Get() );
+			return TMeshAttributesView<ElementIDType, ViewType>(ArraySetPtr->Get());
 		}
 
 		return TMeshAttributesView<ElementIDType, ViewType>();
 	}
 
+	UE_DEPRECATED(4.26, "Please use GetAttributeChannelCount() instead.")
+	int32 GetAttributeIndexCount(const FName AttributeName) const
+	{
+		return GetAttributeChannelCount(AttributeName);
+	}
+
 	/** Returns the number of indices for the attribute with the given name */
+	int32 GetAttributeChannelCount(const FName AttributeName) const
+	{
+		if (const FAttributesSetEntry* ArraySetPtr = this->Map.Find(AttributeName))
+		{
+			return (*ArraySetPtr)->GetNumChannels();
+		}
+
+		return 0;
+	}
+
 	template <typename AttributeType>
-	int32 GetAttributeIndexCount( const FName AttributeName ) const
+	UE_DEPRECATED(4.25, "Please use GetAttributeChannelCount() instead.")
+	int32 GetAttributeIndexCount(const FName AttributeName) const
 	{
-		if( const FAttributesSetEntry* ArraySetPtr = this->Map.Find( AttributeName ) )
+		if (const FAttributesSetEntry* ArraySetPtr = this->Map.Find(AttributeName))
 		{
-			if( ( *ArraySetPtr )->HasType<AttributeType>() )
+			if ((*ArraySetPtr)->HasType<AttributeType>())
 			{
 				using ArrayType = TMeshAttributeArraySet<AttributeType>;
-				return static_cast<const ArrayType*>( ArraySetPtr->Get() )->ArrayType::GetNumIndices();	// note: override virtual dispatch
+				return static_cast<const ArrayType*>( ArraySetPtr->Get() )->ArrayType::GetNumChannels();	// note: override virtual dispatch
 			}
 		}
 
 		return 0;
 	}
 
+	UE_DEPRECATED(4.26, "Please use SetAttributeChannelCount() instead.")
+	void SetAttributeIndexCount(const FName AttributeName, const int32 NumChannels)
+	{
+		SetAttributeChannelCount(AttributeName, NumChannels);
+	}
+
 	/** Sets the number of indices for the attribute with the given name */
+	void SetAttributeChannelCount(const FName AttributeName, const int32 NumChannels)
+	{
+		if (FAttributesSetEntry* ArraySetPtr = this->Map.Find(AttributeName))
+		{
+			(*ArraySetPtr)->SetNumChannels(NumChannels);
+		}
+	}
+
 	template <typename AttributeType>
-	void SetAttributeIndexCount( const FName AttributeName, const int32 NumIndices )
+	UE_DEPRECATED(4.25, "Please use untemplated SetAttributeChannelCount() instead.")
+	void SetAttributeIndexCount(const FName AttributeName, const int32 NumIndices)
 	{
-		if( FAttributesSetEntry* ArraySetPtr = this->Map.Find( AttributeName ) )
+		if (FAttributesSetEntry* ArraySetPtr = this->Map.Find(AttributeName))
 		{
-			if( ( *ArraySetPtr )->HasType<AttributeType>() )
+			if ((*ArraySetPtr)->HasType<AttributeType>())
 			{
 				using ArrayType = TMeshAttributeArraySet<AttributeType>;
-				static_cast<ArrayType*>( ArraySetPtr->Get() )->ArrayType::SetNumIndices( NumIndices );	// note: override virtual dispatch
+				static_cast<ArrayType*>(ArraySetPtr->Get())->ArrayType::SetNumChannels(NumIndices);	// note: override virtual dispatch
 			}
 		}
 	}
 
+	UE_DEPRECATED(4.26, "Please use InsertAttributeChannel() instead.")
+	void InsertAttributeIndex_Old(const FName AttributeName, const int32 Index)
+	{
+		InsertAttributeChannel(AttributeName, Index);
+	}
+
 	/** Insert a new index for the attribute with the given name */
+	void InsertAttributeChannel(const FName AttributeName, const int32 Index)
+	{
+		if (FAttributesSetEntry* ArraySetPtr = this->Map.Find(AttributeName))
+		{
+			(*ArraySetPtr)->InsertChannel(Index);
+		}
+	}
+
 	template <typename AttributeType>
-	void InsertAttributeIndex( const FName AttributeName, const int32 Index )
+	UE_DEPRECATED(4.25, "Please use untemplated InsertAttributeIndexCount() instead.")
+	void InsertAttributeIndex(const FName AttributeName, const int32 Index)
 	{
-		if( FAttributesSetEntry* ArraySetPtr = this->Map.Find( AttributeName ) )
+		if (FAttributesSetEntry* ArraySetPtr = this->Map.Find(AttributeName))
 		{
-			if( ( *ArraySetPtr )->HasType<AttributeType>() )
+			if ((*ArraySetPtr)->HasType<AttributeType>())
 			{
 				using ArrayType = TMeshAttributeArraySet<AttributeType>;
-				static_cast<ArrayType*>( ArraySetPtr->Get() )->ArrayType::InsertIndex( Index );	// note: override virtual dispatch
+				static_cast<ArrayType*>(ArraySetPtr->Get())->ArrayType::InsertChannel(Index);	// note: override virtual dispatch
 			}
 		}
 	}
 
+	UE_DEPRECATED(4.26, "Please use RemoveAttributeChannel() instead.")
+	void RemoveAttributeIndex_Old(const FName AttributeName, const int32 Index)
+	{
+		RemoveAttributeChannel(AttributeName, Index);
+	}
+
 	/** Remove an existing index from the attribute with the given name */
+	void RemoveAttributeChannel(const FName AttributeName, const int32 Index)
+	{
+		if (FAttributesSetEntry* ArraySetPtr = this->Map.Find(AttributeName))
+		{
+			(*ArraySetPtr)->RemoveChannel(Index);
+		}
+	}
+
 	template <typename AttributeType>
-	void RemoveAttributeIndex( const FName AttributeName, const int32 Index )
+	UE_DEPRECATED(4.25, "Please use untemplated RemoveAttributeIndexCount() instead.")
+	void RemoveAttributeIndex(const FName AttributeName, const int32 Index)
 	{
-		if( FAttributesSetEntry* ArraySetPtr = this->Map.Find( AttributeName ) )
+		if (FAttributesSetEntry* ArraySetPtr = this->Map.Find(AttributeName))
 		{
-			if( ( *ArraySetPtr )->HasType<AttributeType>() )
+			if ((*ArraySetPtr)->HasType<AttributeType>())
 			{
 				using ArrayType = TMeshAttributeArraySet<AttributeType>;
-				static_cast<ArrayType*>( ArraySetPtr->Get() )->ArrayType::RemoveIndex( Index );	// note: override virtual dispatch
+				static_cast<ArrayType*>(ArraySetPtr->Get())->ArrayType::RemoveChannel(Index);	// note: override virtual dispatch
 			}
 		}
 	}
 
-	template <typename AttributeType>
-	AttributeType GetAttribute( const ElementIDType ElementID, const FName AttributeName, const int32 AttributeIndex = 0 ) const
+	/**
+	 * Get an attribute value for the given element ID.
+	 * Note: it is generally preferable to get a TMeshAttributesRef and access elements through that, if you wish to access more than one.
+	 */
+	template <typename AttributeType,
+			  typename TEnableIf<!TIsArrayView<AttributeType>::Value, int>::Type = 0>
+	AttributeType GetAttribute(const ElementIDType ElementID, const FName AttributeName, const int32 AttributeChannel = 0) const
 	{
-		const FMeshAttributeArraySetBase* ArraySetPtr = this->Map.FindChecked( AttributeName ).Get();
-		check( ArraySetPtr->HasType<AttributeType>() );
-		return static_cast<const TMeshAttributeArraySet<AttributeType>*>( ArraySetPtr )->GetArrayForIndex( AttributeIndex )[ ElementID.GetValue() ];
+		const FMeshAttributeArraySetBase* ArraySetPtr = this->Map.FindChecked(AttributeName).Get();
+		check(ArraySetPtr->HasType<AttributeType>());
+		check(ArraySetPtr->GetExtent() == 1);
+		const AttributeType* ElementBase = static_cast<const TMeshAttributeArraySet<AttributeType>*>(ArraySetPtr)->GetArrayForChannel(AttributeChannel).GetElementBase(ElementID.GetValue());
+		return ElementBase[0];
 	}
 
-	template <typename AttributeType>
-	void SetAttribute( const ElementIDType ElementID, const FName AttributeName, const int32 AttributeIndex, const AttributeType& AttributeValue )
+	/** Get an attribute value for the given element ID. This is a compound (array) attribute, and returns an array view */
+	template <typename ArrayType,
+			  typename TEnableIf<TIsArrayView<ArrayType>::Value, int>::Type = 0,
+			  typename AttributeType = typename TBreakArrayView<ArrayType>::Type>
+	TArrayView<AttributeType> GetAttribute(const ElementIDType ElementID, const FName AttributeName, const int32 AttributeChannel = 0) const
 	{
-		FMeshAttributeArraySetBase* ArraySetPtr = this->Map.FindChecked( AttributeName ).Get();
-		check( ArraySetPtr->HasType<AttributeType>() );
-		static_cast<TMeshAttributeArraySet<AttributeType>*>( ArraySetPtr )->GetArrayForIndex( AttributeIndex )[ ElementID.GetValue() ] = AttributeValue;
+		const FMeshAttributeArraySetBase* ArraySetPtr = this->Map.FindChecked(AttributeName).Get();
+		check(ArraySetPtr->HasType<AttributeType>());
+		const AttributeType* ElementBase = static_cast<const TMeshAttributeArraySet<AttributeType>*>(ArraySetPtr)->GetArrayForChannel(AttributeChannel).GetElementBase(ElementID.GetValue());
+		return TArrayView<AttributeType>(ElementBase[0], ArraySetPtr->GetExtent());
+	}
+
+
+	/**
+	 * Set an attribute value for the given element ID.
+	 * Note: it is generally preferable to get a TMeshAttributesRef and set multiple elements through that.
+	 */
+	template <typename AttributeType,
+			  typename TEnableIf<!TIsArrayView<AttributeType>::Value, int>::Type = 0>
+	void SetAttribute(const ElementIDType ElementID, const FName AttributeName, const int32 AttributeChannel, const AttributeType& AttributeValue)
+	{
+		FMeshAttributeArraySetBase* ArraySetPtr = this->Map.FindChecked(AttributeName).Get();
+		check(ArraySetPtr->HasType<AttributeType>());
+		check(ArraySetPtr->GetExtent() == 1);
+		AttributeType* ElementBase = static_cast<TMeshAttributeArraySet<AttributeType>*>(ArraySetPtr)->GetArrayForChannel(AttributeChannel).GetElementBase(ElementID.GetValue());
+		ElementBase[0] = AttributeValue;
+	}
+
+	template <typename ArrayType,
+			  typename TEnableIf<TIsArrayView<ArrayType>::Value, int>::Type = 0,
+			  typename AttributeType = typename TRemoveCV<typename TBreakArrayView<ArrayType>::Type>::Type>
+	void SetAttribute(const ElementIDType ElementID, const FName AttributeName, const int32 AttributeChannel, const ArrayType& AttributeValue)
+	{
+		FMeshAttributeArraySetBase* ArraySetPtr = this->Map.FindChecked(AttributeName).Get();
+		check(ArraySetPtr->HasType<AttributeType>());
+		check(ArraySetPtr->GetExtent() == AttributeValue.Num());
+		AttributeType* ElementBase = static_cast<TMeshAttributeArraySet<AttributeType>*>(ArraySetPtr)->GetArrayForChannel(AttributeChannel).GetElementBase(ElementID.GetValue());
+		for (int32 I = 0; I < AttributeValue.Num(); I++)
+		{
+			ElementBase[I] = AttributeValue[I];
+		}
 	}
 
 	/** Inserts a default-initialized value for all attributes of the given ID */
-	FORCEINLINE void Insert( const ElementIDType ElementID )
+	FORCEINLINE void Insert(const ElementIDType ElementID)
 	{
-		this->Insert( ElementID.GetValue() );
+		this->Insert(ElementID.GetValue());
 	}
 
 	/** Removes all attributes with the given ID */
-	FORCEINLINE void Remove( const ElementIDType ElementID )
+	FORCEINLINE void Remove(const ElementIDType ElementID)
 	{
-		this->Remove( ElementID.GetValue() );
+		this->Remove(ElementID.GetValue());
 	}
 
 	/**
 	 * Call the supplied function on each attribute.
 	 * The prototype should be Func( const FName AttributeName, auto AttributesRef );
 	 */
-	template <typename ForEachFunc> void ForEach( ForEachFunc Func );
+	template <typename ForEachFunc> void ForEach(ForEachFunc Func);
 
 	/**
 	* Call the supplied function on each attribute.
 	* The prototype should be Func( const FName AttributeName, auto AttributesConstRef );
 	*/
-	template <typename ForEachFunc> void ForEach( ForEachFunc Func ) const;
-
-	/**
-	 * Call the supplied function on each attribute.
-	 * The prototype should be Func( const FName AttributeName, auto& AttributeIndicesArray );
-	 */
-	template <typename FuncType>
-	UE_DEPRECATED( 4.20, "This is no longer supported; please use ForEach() instead and amend your lambda to accept an auto of type TMeshAttributesRef instead." )
-	void ForEachAttributeIndicesArray( const FuncType& Func )
-	{
-		check( false );
-	}
-
-	template <typename FuncType>
-	UE_DEPRECATED( 4.20, "This is no longer supported; please use ForEach() instead and amend your lambda to accept an auto of type TMeshAttributesConstRef instead." )
-	void ForEachAttributeIndicesArray( const FuncType& Func ) const
-	{
-		check( false );
-	}
-
-	template <typename AttributeType>
-	UE_DEPRECATED( 4.20, "Please use GetAttributesRef() or GetAttributesView() instead." )
-	TMeshAttributeArray<AttributeType, ElementIDType>& GetAttributes( const FName AttributeName, const int32 AttributeIndex = 0 )
-	{
-		FMeshAttributeArraySetBase* ArraySetPtr = this->Map.FindChecked( AttributeName ).Get();
-		check( ArraySetPtr->HasType<AttributeType>() );
-		return static_cast<TAttributeIndicesArray<AttributeType, ElementIDType>*>( ArraySetPtr )->GetArrayForIndex( AttributeIndex );
-	}
-
-	template <typename AttributeType>
-	UE_DEPRECATED( 4.20, "Please use GetAttributesRef() or GetAttributesView() instead." )
-	const TMeshAttributeArray<AttributeType, ElementIDType>& GetAttributes( const FName AttributeName, const int32 AttributeIndex = 0 ) const
-	{
-		const FMeshAttributeArraySetBase* ArraySetPtr = this->Map.FindChecked( AttributeName ).Get();
-		check( ArraySetPtr->HasType<AttributeType>() );
-		return static_cast<const TAttributeIndicesArray<AttributeType, ElementIDType>*>( ArraySetPtr )->GetArrayForIndex( AttributeIndex );
-	}
-
-	template <typename AttributeType>
-	UE_DEPRECATED( 4.20, "Please use GetAttributesRef() or GetAttributesView() instead." )
-	TAttributeIndicesArray<AttributeType, ElementIDType>& GetAttributesSet( const FName AttributeName )
-	{
-		FMeshAttributeArraySetBase* ArraySetPtr = this->Map.FindChecked( AttributeName ).Get();
-		check( ArraySetPtr->HasType<AttributeType>() );
-		return static_cast<TAttributeIndicesArray<AttributeType, ElementIDType>&>( *ArraySetPtr );
-	}
-
-	template <typename AttributeType>
-	UE_DEPRECATED( 4.20, "Please use GetAttributesRef() or GetAttributesView() instead." )
-	const TAttributeIndicesArray<AttributeType, ElementIDType>& GetAttributesSet( const FName AttributeName ) const
-	{
-		const FMeshAttributeArraySetBase* ArraySetPtr = this->Map.FindChecked( AttributeName ).Get();
-		check( ArraySetPtr->HasType<AttributeType>() );
-		return static_cast<const TAttributeIndicesArray<AttributeType, ElementIDType>&>( *ArraySetPtr );
-	}
+	template <typename ForEachFunc> void ForEach(ForEachFunc Func) const;
 };
 
 
@@ -1187,36 +1794,44 @@ namespace ForEachImpl
 {
 	// Declare type of jump table used to dispatch functions
 	template <typename ElementIDType, typename ForEachFunc>
-	using JumpTableType = TJumpTable<void( FName, ForEachFunc, FMeshAttributeArraySetBase* ), TTupleArity<AttributeTypes>::Value>;
+	using JumpTableType = TJumpTable<void(FName, ForEachFunc, FMeshAttributeArraySetBase*), TTupleArity<AttributeTypes>::Value>;
 
 	// Define dispatch function
 	template <typename ElementIDType, typename ForEachFunc, uint32 I>
-	static void Dispatch( FName Name, ForEachFunc Fn, FMeshAttributeArraySetBase* Attributes )
+	static void Dispatch(FName Name, ForEachFunc Fn, FMeshAttributeArraySetBase* Attributes)
 	{
 		using AttributeType = typename TTupleElement<I, AttributeTypes>::Type;
-		Fn( Name, TMeshAttributesRef<ElementIDType, AttributeType>( static_cast<TMeshAttributeArraySet<AttributeType>*>( Attributes ) ) );
+		if (Attributes->GetExtent() == 1)
+		{
+			Fn(Name, TMeshAttributesRef<ElementIDType, AttributeType>(static_cast<TMeshAttributeArraySet<AttributeType>*>(Attributes)));
+		}
+		else
+		{
+			// @todo: allow ForEach to iterate through array types
+//			Fn(Name, TMeshAttributesRef<ElementIDType, TArrayView<AttributeType>(static_cast<TMeshAttributeArraySet<AttributeType>*>(Attributes), Attributes->GetExtent()));
+		}
 	}
 
 	// Build ForEach jump table at compile time, a separate instantiation of Dispatch for each attribute type
 	template <typename ElementIDType, typename ForEachFunc, uint32... Is>
-	static constexpr JumpTableType<ElementIDType, ForEachFunc> MakeJumpTable( TIntegerSequence< uint32, Is...> )
+	static constexpr JumpTableType<ElementIDType, ForEachFunc> MakeJumpTable(TIntegerSequence< uint32, Is...>)
 	{
-		return JumpTableType<ElementIDType, ForEachFunc>( Dispatch<ElementIDType, ForEachFunc, Is>... );
+		return JumpTableType<ElementIDType, ForEachFunc>(Dispatch<ElementIDType, ForEachFunc, Is>...);
 	}
 }
 
 template <typename ElementIDType>
 template <typename ForEachFunc>
-void TAttributesSet<ElementIDType>::ForEach( ForEachFunc Func )
+void TAttributesSet<ElementIDType>::ForEach(ForEachFunc Func)
 {
 	// Construct compile-time jump table for dispatching ForEachImpl::Dispatch() by the attribute type at runtime
 	static constexpr ForEachImpl::JumpTableType<ElementIDType, ForEachFunc>
-		JumpTable = ForEachImpl::MakeJumpTable<ElementIDType, ForEachFunc>( TMakeIntegerSequence<uint32, TTupleArity<AttributeTypes>::Value>() );
+		JumpTable = ForEachImpl::MakeJumpTable<ElementIDType, ForEachFunc>(TMakeIntegerSequence<uint32, TTupleArity<AttributeTypes>::Value>());
 
-	for( auto& MapEntry : this->Map )
+	for (auto& MapEntry : this->Map)
 	{
 		const uint32 Type = MapEntry.Value->GetType();
-		JumpTable.Fns[ Type ]( MapEntry.Key, Func, MapEntry.Value.Get() );
+		JumpTable.Fns[Type](MapEntry.Key, Func, MapEntry.Value.Get());
 	}
 }
 
@@ -1224,36 +1839,44 @@ namespace ForEachConstImpl
 {
 	// Declare type of jump table used to dispatch functions
 	template <typename ElementIDType, typename ForEachFunc>
-	using JumpTableType = TJumpTable<void( FName, ForEachFunc, const FMeshAttributeArraySetBase* ), TTupleArity<AttributeTypes>::Value>;
+	using JumpTableType = TJumpTable<void(FName, ForEachFunc, const FMeshAttributeArraySetBase*), TTupleArity<AttributeTypes>::Value>;
 
 	// Define dispatch function
 	template <typename ElementIDType, typename ForEachFunc, uint32 I>
-	static void Dispatch( FName Name, ForEachFunc Fn, const FMeshAttributeArraySetBase* Attributes )
+	static void Dispatch(FName Name, ForEachFunc Fn, const FMeshAttributeArraySetBase* Attributes)
 	{
 		using AttributeType = typename TTupleElement<I, AttributeTypes>::Type;
-		Fn( Name, TMeshAttributesConstRef<ElementIDType, AttributeType>( static_cast<const TMeshAttributeArraySet<AttributeType>*>( Attributes ) ) );
+		if (Attributes->GetExtent() == 1)
+		{
+			Fn(Name, TMeshAttributesConstRef<ElementIDType, AttributeType>(static_cast<const TMeshAttributeArraySet<AttributeType>*>(Attributes)));
+		}
+		else
+		{
+			// @todo: allow ForEach to iterate through array types
+//			Fn(Name, TMeshAttributesConstRef<ElementIDType, TArrayView<const AttributeType>>(static_cast<const TMeshAttributeArraySet<AttributeType>*>(Attributes), Attributes->GetExtent()));
+		}
 	}
 
 	// Build ForEach jump table at compile time, a separate instantiation of Dispatch for each attribute type
 	template <typename ElementIDType, typename ForEachFunc, uint32... Is>
-	static constexpr JumpTableType<ElementIDType, ForEachFunc> MakeJumpTable( TIntegerSequence< uint32, Is...> )
+	static constexpr JumpTableType<ElementIDType, ForEachFunc> MakeJumpTable(TIntegerSequence< uint32, Is...>)
 	{
-		return JumpTableType<ElementIDType, ForEachFunc>( Dispatch<ElementIDType, ForEachFunc, Is>... );
+		return JumpTableType<ElementIDType, ForEachFunc>(Dispatch<ElementIDType, ForEachFunc, Is>...);
 	}
 }
 
 template <typename ElementIDType>
 template <typename ForEachFunc>
-void TAttributesSet<ElementIDType>::ForEach( ForEachFunc Func ) const
+void TAttributesSet<ElementIDType>::ForEach(ForEachFunc Func) const
 {
 	// Construct compile-time jump table for dispatching ForEachImpl::Dispatch() by the attribute type at runtime
 	static constexpr ForEachConstImpl::JumpTableType<ElementIDType, ForEachFunc>
-		JumpTable = ForEachConstImpl::MakeJumpTable<ElementIDType, ForEachFunc>( TMakeIntegerSequence<uint32, TTupleArity<AttributeTypes>::Value>() );
+		JumpTable = ForEachConstImpl::MakeJumpTable<ElementIDType, ForEachFunc>(TMakeIntegerSequence<uint32, TTupleArity<AttributeTypes>::Value>());
 
-	for( const auto& MapEntry : this->Map )
+	for (const auto& MapEntry : this->Map)
 	{
 		const uint32 Type = MapEntry.Value->GetType();
-		JumpTable.Fns[ Type ]( MapEntry.Key, Func, MapEntry.Value.Get() );
+		JumpTable.Fns[Type](MapEntry.Key, Func, MapEntry.Value.Get());
 	}
 }
 
@@ -1266,28 +1889,28 @@ void TAttributesSet<ElementIDType>::ForEach( ForEachFunc Func ) const
 namespace CreateTypeImpl
 {
 	// Declare type of jump table used to dispatch functions
-	using JumpTableType = TJumpTable<TUniquePtr<FMeshAttributeArraySetBase>(), TTupleArity<AttributeTypes>::Value>;
+	using JumpTableType = TJumpTable<TUniquePtr<FMeshAttributeArraySetBase>(uint32), TTupleArity<AttributeTypes>::Value>;
 
 	// Define dispatch function
 	template <uint32 I>
-	static TUniquePtr<FMeshAttributeArraySetBase> Dispatch()
+	static TUniquePtr<FMeshAttributeArraySetBase> Dispatch(uint32 Extent)
 	{
 		using AttributeType = typename TTupleElement<I, AttributeTypes>::Type;
-		return MakeUnique<TMeshAttributeArraySet<AttributeType>>();
+		return MakeUnique<TMeshAttributeArraySet<AttributeType>>(Extent);
 	}
 
 	// Build RegisterAttributeOfType jump table at compile time, a separate instantiation of Dispatch for each attribute type
 	template <uint32... Is>
-	static constexpr JumpTableType MakeJumpTable( TIntegerSequence< uint32, Is...> )
+	static constexpr JumpTableType MakeJumpTable(TIntegerSequence< uint32, Is...>)
 	{
-		return JumpTableType( Dispatch<Is>... );
+		return JumpTableType(Dispatch<Is>...);
 	}
 }
 
-inline void FAttributesSetEntry::CreateArrayOfType( const uint32 Type )
+inline void FAttributesSetEntry::CreateArrayOfType(const uint32 Type, const uint32 Extent)
 {
-	static constexpr CreateTypeImpl::JumpTableType JumpTable = CreateTypeImpl::MakeJumpTable( TMakeIntegerSequence<uint32, TTupleArity<AttributeTypes>::Value>() );
-	Ptr = JumpTable.Fns[ Type ]();
+	static constexpr CreateTypeImpl::JumpTableType JumpTable = CreateTypeImpl::MakeJumpTable(TMakeIntegerSequence<uint32, TTupleArity<AttributeTypes>::Value>());
+	Ptr = JumpTable.Fns[Type](Extent);
 }
 
 
@@ -1318,7 +1941,7 @@ namespace AttributesViewGetImpl
 	static ViewType Dispatch( const FMeshAttributeArraySetBase* Array, const int32 Index )
 	{
 		// Implementation when the attribute type is convertible to the view type
-		return ViewType( static_cast<const TMeshAttributeArraySet<typename TTupleElement<I, AttributeTypes>::Type>*>( Array )->GetArrayForIndex( 0 )[ Index ] );
+		return ViewType( static_cast<const TMeshAttributeArraySet<typename TTupleElement<I, AttributeTypes>::Type>*>( Array )->GetArrayForChannel( 0 )[ Index ] );
 	}
 
 	template <typename ViewType, uint32 I, typename TEnableIf<!TIsViewable<ViewType, I>::Value, int>::Type = 0>
@@ -1360,7 +1983,7 @@ namespace AttributesViewGetWithIndexImpl
 	template <typename ViewType, uint32 I, typename TEnableIf<TIsViewable<ViewType, I>::Value, int>::Type = 0>
 	static ViewType Dispatch( const FMeshAttributeArraySetBase* Array, const int32 ElementIndex, const int32 AttributeIndex )
 	{
-		return ViewType( static_cast<const TMeshAttributeArraySet<typename TTupleElement<I, AttributeTypes>::Type>*>( Array )->GetArrayForIndex( AttributeIndex )[ ElementIndex ] );
+		return ViewType( static_cast<const TMeshAttributeArraySet<typename TTupleElement<I, AttributeTypes>::Type>*>( Array )->GetArrayForChannel( AttributeIndex )[ ElementIndex ] );
 	}
 
 	template <typename ViewType, uint32 I, typename TEnableIf<!TIsViewable<ViewType, I>::Value, int>::Type = 0>
@@ -1403,7 +2026,7 @@ namespace AttributesViewSetImpl
 	{
 		// Implementation when the attribute type is convertible to the view type
 		using AttributeType = typename TTupleElement<I, AttributeTypes>::Type;
-		static_cast<TMeshAttributeArraySet<AttributeType>*>( Array )->GetArrayForIndex( 0 )[ Index ] = AttributeType( Value );
+		static_cast<TMeshAttributeArraySet<AttributeType>*>( Array )->GetArrayForChannel( 0 )[ Index ] = AttributeType( Value );
 	}
 
 	template <typename ViewType, uint32 I, typename TEnableIf<!TIsViewable<ViewType, I>::Value, int>::Type = 0>
@@ -1446,7 +2069,7 @@ namespace AttributesViewSetWithIndexImpl
 	{
 		// Implementation when the attribute type is convertible to the view type
 		using AttributeType = typename TTupleElement<I, AttributeTypes>::Type;
-		static_cast<TMeshAttributeArraySet<AttributeType>*>( Array )->GetArrayForIndex( AttributeIndex )[ ElementIndex ] = AttributeType( Value );
+		static_cast<TMeshAttributeArraySet<AttributeType>*>( Array )->GetArrayForChannel( AttributeIndex )[ ElementIndex ] = AttributeType( Value );
 	}
 
 	template <typename ViewType, uint32 I, typename TEnableIf<!TIsViewable<ViewType, I>::Value, int>::Type = 0>
