First pass of MeshDescription API and format refactor.

diff --git a/Engine/Source/Runtime/MeshDescription/Public/MeshAttributes.h b/Engine/Source/Runtime/MeshDescription/Public/MeshAttributes.h
--- a/Engine/Source/Runtime/MeshDescription/Public/MeshAttributes.h
+++ b/Engine/Source/Runtime/MeshDescription/Public/MeshAttributes.h
@@ -11,6 +11,36 @@ namespace MeshAttribute
 	{
 		extern MESHDESCRIPTION_API const FName Position;
 	}
+
+	namespace VertexInstance
+	{
+		extern MESHDESCRIPTION_API const FName VertexIndex;
+	}
+
+	namespace Edge
+	{
+		extern MESHDESCRIPTION_API const FName VertexIndex;
+	}
+
+	namespace Triangle
+	{
+		extern MESHDESCRIPTION_API const FName VertexInstanceIndex;
+		extern MESHDESCRIPTION_API const FName PolygonIndex;
+		extern MESHDESCRIPTION_API const FName EdgeIndex;
+		extern MESHDESCRIPTION_API const FName VertexIndex;
+		extern MESHDESCRIPTION_API const FName UVIndex;
+		extern MESHDESCRIPTION_API const FName PolygonGroupIndex;
+	}
+
+	namespace UV
+	{
+		extern MESHDESCRIPTION_API const FName UVCoordinate;
+	}
+
+	namespace Polygon
+	{
+		extern MESHDESCRIPTION_API const FName PolygonGroupIndex;
+	}
 }
 
 
@@ -26,8 +56,103 @@ public:
 	virtual void Register();
 
 	/** Accessors for cached vertex position array */
-	TVertexAttributesRef<FVector> GetVertexPositions() { return MeshDescription.VertexAttributes().GetAttributesRef<FVector>(MeshAttribute::Vertex::Position); }
-	TVertexAttributesConstRef<FVector> GetVertexPositions() const { return MeshDescription.VertexAttributes().GetAttributesRef<FVector>(MeshAttribute::Vertex::Position); }
+	TVertexAttributesRef<FVector> GetVertexPositions()
+	{
+		return MeshDescription.VertexAttributes().GetAttributesRef<FVector>(MeshAttribute::Vertex::Position);
+	}
+
+	TVertexAttributesConstRef<FVector> GetVertexPositions() const
+	{
+		return MeshDescription.VertexAttributes().GetAttributesRef<FVector>(MeshAttribute::Vertex::Position);
+	}
+
+	/** Accessors for array of vertex IDs for vertex instances */
+	TVertexInstanceAttributesRef<FVertexID> GetVertexInstanceVertexIndices()
+	{
+		return MeshDescription.VertexInstanceAttributes().GetAttributesRef<FVertexID>(MeshAttribute::VertexInstance::VertexIndex);
+	}
+
+	TVertexInstanceAttributesConstRef<FVertexID> GetVertexInstanceVertexIndices() const
+	{
+		return MeshDescription.VertexInstanceAttributes().GetAttributesRef<FVertexID>(MeshAttribute::VertexInstance::VertexIndex);
+	}
+
+	/** Accessors for array of vertex IDs for edges */
+	TEdgeAttributesRef<TArrayView<FVertexID>> GetEdgeVertexIndices()
+	{
+		return MeshDescription.EdgeAttributes().GetAttributesRef<TArrayView<FVertexID>>(MeshAttribute::Edge::VertexIndex);
+	}
+	
+	TEdgeAttributesConstRef<TArrayView<FVertexID>> GetEdgeVertexIndices() const
+	{
+		return MeshDescription.EdgeAttributes().GetAttributesRef<TArrayView<FVertexID>>(MeshAttribute::Edge::VertexIndex);
+	}
+
+	/** Accessors for array of vertex instance IDs for triangles */
+	TTriangleAttributesRef<TArrayView<FVertexInstanceID>> GetTriangleVertexInstanceIndices()
+	{
+		return MeshDescription.TriangleAttributes().GetAttributesRef<TArrayView<FVertexInstanceID>>(MeshAttribute::Triangle::VertexInstanceIndex);
+	}
+
+	TTriangleAttributesConstRef<TArrayView<FVertexInstanceID>> GetTriangleVertexInstanceIndices() const
+	{
+		return MeshDescription.TriangleAttributes().GetAttributesRef<TArrayView<FVertexInstanceID>>(MeshAttribute::Triangle::VertexInstanceIndex);
+	}
+
+	/** Accessors for array of edge IDs for triangles */
+	TTriangleAttributesRef<TArrayView<FEdgeID>> GetTriangleEdgeIndices()
+	{
+		return MeshDescription.TriangleAttributes().GetAttributesRef<TArrayView<FEdgeID>>(MeshAttribute::Triangle::EdgeIndex);
+	}
+
+	TTriangleAttributesConstRef<TArrayView<FEdgeID>> GetTriangleEdgeIndices() const
+	{
+		return MeshDescription.TriangleAttributes().GetAttributesRef<TArrayView<FEdgeID>>(MeshAttribute::Triangle::EdgeIndex);
+	}
+
+	/** Accessors for array of vertex IDs for triangles */
+	TTriangleAttributesRef<TArrayView<FVertexID>> GetTriangleVertexIndices()
+	{
+		return MeshDescription.TriangleAttributes().GetAttributesRef<TArrayView<FVertexID>>(MeshAttribute::Triangle::VertexIndex);
+	}
+
+	TTriangleAttributesConstRef<TArrayView<FVertexID>> GetTriangleVertexIndices() const
+	{
+		return MeshDescription.TriangleAttributes().GetAttributesRef<TArrayView<FVertexID>>(MeshAttribute::Triangle::VertexIndex);
+	}
+
+	/** Accessors for array of UV IDs for triangles */
+	TTriangleAttributesRef<TArrayView<FUVID>> GetTriangleUVIndices()
+	{
+		return MeshDescription.TriangleAttributes().GetAttributesRef<TArrayView<FUVID>>(MeshAttribute::Triangle::UVIndex);
+	}
+
+	TTriangleAttributesConstRef<TArrayView<FUVID>> GetTriangleUVIndices() const
+	{
+		return MeshDescription.TriangleAttributes().GetAttributesRef<TArrayView<FUVID>>(MeshAttribute::Triangle::UVIndex);
+	}
+
+	/** Accessors for array of polygon group IDs for triangles */
+	TTriangleAttributesRef<FPolygonGroupID> GetTrianglePolygonGroupIndices()
+	{
+		return MeshDescription.TriangleAttributes().GetAttributesRef<FPolygonGroupID>(MeshAttribute::Triangle::PolygonGroupIndex);
+	}
+
+	TTriangleAttributesConstRef<FPolygonGroupID> GetTrianglePolygonGroupIndices() const
+	{
+		return MeshDescription.TriangleAttributes().GetAttributesRef<FPolygonGroupID>(MeshAttribute::Triangle::PolygonGroupIndex);
+	}
+
+	/** Accessors for array of polygon group IDs for polygons */
+	TPolygonAttributesRef<FPolygonGroupID> GetPolygonPolygonGroupIndices()
+	{
+		return MeshDescription.PolygonAttributes().GetAttributesRef<FPolygonGroupID>(MeshAttribute::Polygon::PolygonGroupIndex);
+	}
+
+	TPolygonAttributesConstRef<FPolygonGroupID> GetPolygonPolygonGroupIndices() const
+	{
+		return MeshDescription.PolygonAttributes().GetAttributesRef<FPolygonGroupID>(MeshAttribute::Polygon::PolygonGroupIndex);
+	}
 
 protected:
 
@@ -45,6 +170,54 @@ public:
 	/** Accessors for cached vertex position array */
 	TVertexAttributesConstRef<FVector> GetVertexPositions() const { return MeshDescription.VertexAttributes().GetAttributesRef<FVector>(MeshAttribute::Vertex::Position); }
 
+	/** Accessors for array of vertex IDs for vertex instances */
+	TVertexInstanceAttributesConstRef<FVertexID> GetVertexInstanceVertexIndices() const
+	{
+		return MeshDescription.VertexInstanceAttributes().GetAttributesRef<FVertexID>(MeshAttribute::VertexInstance::VertexIndex);
+	}
+
+	/** Accessors for array of vertex IDs for edges */
+	TEdgeAttributesConstRef<TArrayView<FVertexID>> GetEdgeVertexIndices() const
+	{
+		return MeshDescription.EdgeAttributes().GetAttributesRef<TArrayView<FVertexID>>(MeshAttribute::Edge::VertexIndex);
+	}
+
+	/** Accessors for array of vertex instance IDs for triangles */
+	TTriangleAttributesConstRef<TArrayView<FVertexInstanceID>> GetTriangleVertexInstanceIndices() const
+	{
+		return MeshDescription.TriangleAttributes().GetAttributesRef<TArrayView<FVertexInstanceID>>(MeshAttribute::Triangle::VertexInstanceIndex);
+	}
+
+	/** Accessors for array of edge IDs for triangles */
+	TTriangleAttributesConstRef<TArrayView<FEdgeID>> GetTriangleEdgeIndices() const
+	{
+		return MeshDescription.TriangleAttributes().GetAttributesRef<TArrayView<FEdgeID>>(MeshAttribute::Triangle::EdgeIndex);
+	}
+
+	/** Accessors for array of vertex IDs for triangles */
+	TTriangleAttributesConstRef<TArrayView<FVertexID>> GetTriangleVertexIndices() const
+	{
+		return MeshDescription.TriangleAttributes().GetAttributesRef<TArrayView<FVertexID>>(MeshAttribute::Triangle::VertexIndex);
+	}
+
+	/** Accessors for array of UV IDs for triangles */
+	TTriangleAttributesConstRef<TArrayView<FUVID>> GetTriangleUVIndices() const
+	{
+		return MeshDescription.TriangleAttributes().GetAttributesRef<TArrayView<FUVID>>(MeshAttribute::Triangle::UVIndex);
+	}
+
+	/** Accessors for array of polygon group IDs for triangles */
+	TTriangleAttributesConstRef<FPolygonGroupID> GetTrianglePolygonGroupIndices() const
+	{
+		return MeshDescription.TriangleAttributes().GetAttributesRef<FPolygonGroupID>(MeshAttribute::Triangle::PolygonGroupIndex);
+	}
+
+	/** Accessors for array of polygon group IDs for polygons */
+	TPolygonAttributesConstRef<FPolygonGroupID> GetPolygonPolygonGroupIndices() const
+	{
+		return MeshDescription.PolygonAttributes().GetAttributesRef<FPolygonGroupID>(MeshAttribute::Polygon::PolygonGroupIndex);
+	}
+
 protected:
 
 	const FMeshDescription& MeshDescription;
diff --git a/Engine/Source/Runtime/MeshDescription/Public/MeshDescription.h b/Engine/Source/Runtime/MeshDescription/Public/MeshDescription.h
--- a/Engine/Source/Runtime/MeshDescription/Public/MeshDescription.h
+++ b/Engine/Source/Runtime/MeshDescription/Public/MeshDescription.h
@@ -8,8 +8,10 @@
 #include "UObject/ObjectMacros.h"
 #include "MeshTypes.h"
 #include "MeshElementArray.h"
-#include "MeshAttributeArray.h"
+#include "MeshElementContainer.h"
+#include "MeshElementIndexer.h"
 #include "Algo/Accumulate.h"
+#include "Algo/Copy.h"
 #include "Algo/Find.h"
 #include "Containers/ArrayView.h"
 #include "UObject/EditorObjectVersion.h"
@@ -25,268 +27,21 @@ enum
 	MAX_MESH_TEXTURE_COORDS_MD = 8,
 };
 
-struct FMeshVertex
-{
-	friend struct FMeshDescription;
-
-	FMeshVertex()
-	{}
-
-private:
-
-	/** All of vertex instances which reference this vertex (for split vertex support) */
-	TArray<FVertexInstanceID> VertexInstanceIDs;
-
-	/** The edges connected to this vertex */
-	TArray<FEdgeID> ConnectedEdgeIDs;
-
-public:
-
-	/** Serializer */
-	friend FArchive& operator<<(FArchive& Ar, FMeshVertex& Vertex)
-	{
-		if (Ar.IsLoading() && Ar.CustomVer(FReleaseObjectVersion::GUID) < FReleaseObjectVersion::MeshDescriptionNewSerialization)
-		{
-			Ar << Vertex.VertexInstanceIDs;
-			Ar << Vertex.ConnectedEdgeIDs;
-		}
-
-		return Ar;
-	}
-};
-
-
-struct FMeshVertexInstance
-{
-	friend struct FMeshDescription;
-
-	FMeshVertexInstance()
-		: VertexID(FVertexID::Invalid)
-	{}
-
-private:
-
-	/** The vertex this is instancing */
-	FVertexID VertexID;
-
-	/** List of connected triangles */
-	TArray<FTriangleID> ConnectedTriangles;
-
-public:
-
-	/** Serializer */
-	friend FArchive& operator<<(FArchive& Ar, FMeshVertexInstance& VertexInstance)
-	{
-		Ar << VertexInstance.VertexID;
-		if (Ar.IsLoading() && Ar.CustomVer(FReleaseObjectVersion::GUID) < FReleaseObjectVersion::MeshDescriptionNewSerialization)
-		{
-			TArray<FPolygonID> ConnectedPolygons_DISCARD;
-			Ar << ConnectedPolygons_DISCARD;
-		}
-
-		return Ar;
-	}
-};
-
-
-struct FMeshEdge
-{
-	friend struct FMeshDescription;
-
-	FMeshEdge()
-	{
-		VertexIDs[0] = FVertexID::Invalid;
-		VertexIDs[1] = FVertexID::Invalid;
-	}
-
-private:
-
-	/** IDs of the two editable mesh vertices that make up this edge.  The winding direction is not defined. */
-	FVertexID VertexIDs[2];
-
-	/** The triangles that share this edge */
-	TArray<FTriangleID> ConnectedTriangles;
-
-public:
-
-	/** Serializer */
-	friend FArchive& operator<<(FArchive& Ar, FMeshEdge& Edge)
-	{
-		Ar << Edge.VertexIDs[0];
-		Ar << Edge.VertexIDs[1];
-		if (Ar.IsLoading() && Ar.CustomVer(FReleaseObjectVersion::GUID) < FReleaseObjectVersion::MeshDescriptionNewSerialization)
-		{
-			TArray<FPolygonID> ConnectedPolygons_DISCARD;
-			Ar << ConnectedPolygons_DISCARD;
-		}
-
-		return Ar;
-	}
-};
-
-
-struct FMeshTriangle
-{
-	friend struct FMeshDescription;
-
-	FMeshTriangle()
-	{
-		VertexInstanceIDs[0] = FVertexInstanceID::Invalid;
-		VertexInstanceIDs[1] = FVertexInstanceID::Invalid;
-		VertexInstanceIDs[2] = FVertexInstanceID::Invalid;
-	}
-
-private:
-
-	/** Vertex instance IDs that make up this triangle.  Indices must be ordered counter-clockwise. */
-	FVertexInstanceID VertexInstanceIDs[3];
-
-	/** Polygon which contains this triangle */
-	FPolygonID PolygonID;
-
-public:
-
-	/** Gets the specified triangle vertex instance ID.  Pass an index between 0 and 2 inclusive. */
-	inline FVertexInstanceID GetVertexInstanceID(const int32 Index) const
-	{
-		checkSlow(Index >= 0 && Index <= 2);
-		return VertexInstanceIDs[Index];
-	}
-
-	/** Sets the specified triangle vertex instance ID.  Pass an index between 0 and 2 inclusive, and the new vertex instance ID to store. */
-	inline void SetVertexInstanceID(const int32 Index, const FVertexInstanceID NewVertexInstanceID)
-	{
-		// When we deprecate direct member access, this will be a simple array lookup
-		checkSlow(Index >= 0 && Index <= 2);
-		VertexInstanceIDs[Index] = NewVertexInstanceID;
-	}
-
-	/** Serializer */
-	friend FArchive& operator<<(FArchive& Ar, FMeshTriangle& Triangle)
-	{
-		Ar << Triangle.VertexInstanceIDs[0];
-		Ar << Triangle.VertexInstanceIDs[1];
-		Ar << Triangle.VertexInstanceIDs[2];
-
-		if (!Ar.IsLoading() || Ar.CustomVer(FEditorObjectVersion::GUID) >= FEditorObjectVersion::MeshDescriptionTriangles)
-		{
-			Ar << Triangle.PolygonID;
-		}
-
-		return Ar;
-	}
-};
-
-
-struct FMeshPolygon
-{
-	friend struct FMeshDescription;
-
-	FMeshPolygon()
-		: PolygonGroupID(FPolygonGroupID::Invalid)
-	{}
-
-private:
-
-	/** The outer boundary edges of this polygon */
-	TArray<FVertexInstanceID> VertexInstanceIDs;
-
-	/** List of triangle IDs which make up this polygon */
-	TArray<FTriangleID> TriangleIDs;
-
-	/** The polygon group which contains this polygon */
-	FPolygonGroupID PolygonGroupID;
-
-public:
-
-	/** Serializer */
-	friend FArchive& operator<<(FArchive& Ar, FMeshPolygon& Polygon)
-	{
-		if (Ar.IsSaving() &&
-			Ar.CustomVer(FEditorObjectVersion::GUID) >= FEditorObjectVersion::MeshDescriptionTriangles &&
-			Polygon.VertexInstanceIDs.Num() == 3)
-		{
-			// Optimisation: if polygon is a triangle, don't serialize the vertices as they can be copied over from the associated triangle
-			TArray<FVertexInstanceID> Empty;
-			Ar << Empty;
-		}
-		else
-		{
-			Ar << Polygon.VertexInstanceIDs;
-		}
-
-		if (Ar.IsLoading() && Ar.CustomVer(FEditorObjectVersion::GUID) < FEditorObjectVersion::MeshDescriptionRemovedHoles)
-		{
-			TArray<TArray<FVertexInstanceID>> Empty;
-			Ar << Empty;
-		}
-
-		if (Ar.IsLoading() && Ar.CustomVer(FReleaseObjectVersion::GUID) < FReleaseObjectVersion::MeshDescriptionNewSerialization)
-		{
-			TArray<FMeshTriangle> Triangles_DISCARD;
-			Ar << Triangles_DISCARD;
-		}
-
-		Ar << Polygon.PolygonGroupID;
-
-		return Ar;
-	}
-};
-
-
-struct FMeshPolygonGroup
-{
-	friend struct FMeshDescription;
-
-	FMeshPolygonGroup()
-	{}
-
-private:
-
-	/** All polygons in this group */
-	TArray<FPolygonID> Polygons;
-
-public:
-
-	/** Serializer */
-	friend FArchive& operator<<(FArchive& Ar, FMeshPolygonGroup& PolygonGroup)
-	{
-		if (Ar.IsLoading() && Ar.CustomVer(FReleaseObjectVersion::GUID) < FReleaseObjectVersion::MeshDescriptionNewSerialization)
-		{
-			Ar << PolygonGroup.Polygons;
-		}
-
-		return Ar;
-	}
-};
-
 
 /** Define container types */
-using FVertexArray = TMeshElementArray<FMeshVertex, FVertexID>;
-using FVertexInstanceArray = TMeshElementArray<FMeshVertexInstance, FVertexInstanceID>;
-using FEdgeArray = TMeshElementArray<FMeshEdge, FEdgeID>;
-using FTriangleArray = TMeshElementArray<FMeshTriangle, FTriangleID>;
-using FPolygonArray = TMeshElementArray<FMeshPolygon, FPolygonID>;
-using FPolygonGroupArray = TMeshElementArray<FMeshPolygonGroup, FPolygonGroupID>;
+using FVertexArray = TMeshElementContainer<FVertexID>;
+using FVertexInstanceArray = TMeshElementContainer<FVertexInstanceID>;
+using FEdgeArray = TMeshElementContainer<FEdgeID>;
+using FUVArray = TMeshElementContainer<FUVID>;
+using FTriangleArray = TMeshElementContainer<FTriangleID>;
+using FPolygonArray = TMeshElementContainer<FPolygonID>;
+using FPolygonGroupArray = TMeshElementContainer<FPolygonGroupID>;
 
 /** Define aliases for element attributes */
-template <typename AttributeType> using TVertexAttributeIndicesArray = TAttributeIndicesArray<AttributeType, FVertexID>;
-template <typename AttributeType> using TVertexInstanceAttributeIndicesArray = TAttributeIndicesArray<AttributeType, FVertexInstanceID>;
-template <typename AttributeType> using TEdgeAttributeIndicesArray = TAttributeIndicesArray<AttributeType, FEdgeID>;
-template <typename AttributeType> using TTriangleAttributeIndicesArray = TAttributeIndicesArray<AttributeType, FTriangleID>;
-template <typename AttributeType> using TPolygonAttributeIndicesArray = TAttributeIndicesArray<AttributeType, FPolygonID>;
-template <typename AttributeType> using TPolygonGroupAttributeIndicesArray = TAttributeIndicesArray<AttributeType, FPolygonGroupID>;
-
-template <typename AttributeType> using TVertexAttributeArray = TMeshAttributeArray<AttributeType, FVertexID>;
-template <typename AttributeType> using TVertexInstanceAttributeArray = TMeshAttributeArray<AttributeType, FVertexInstanceID>;
-template <typename AttributeType> using TEdgeAttributeArray = TMeshAttributeArray<AttributeType, FEdgeID>;
-template <typename AttributeType> using TTriangleAttributeArray = TMeshAttributeArray<AttributeType, FTriangleID>;
-template <typename AttributeType> using TPolygonAttributeArray = TMeshAttributeArray<AttributeType, FPolygonID>;
-template <typename AttributeType> using TPolygonGroupAttributeArray = TMeshAttributeArray<AttributeType, FPolygonGroupID>;
-
 template <typename AttributeType> using TVertexAttributesRef = TMeshAttributesRef<FVertexID, AttributeType>;
 template <typename AttributeType> using TVertexInstanceAttributesRef = TMeshAttributesRef<FVertexInstanceID, AttributeType>;
 template <typename AttributeType> using TEdgeAttributesRef = TMeshAttributesRef<FEdgeID, AttributeType>;
+template <typename AttributeType> using TUVAttributesRef = TMeshAttributesRef<FUVID, AttributeType>;
 template <typename AttributeType> using TTriangleAttributesRef = TMeshAttributesRef<FTriangleID, AttributeType>;
 template <typename AttributeType> using TPolygonAttributesRef = TMeshAttributesRef<FPolygonID, AttributeType>;
 template <typename AttributeType> using TPolygonGroupAttributesRef = TMeshAttributesRef<FPolygonGroupID, AttributeType>;
@@ -294,6 +49,7 @@ template <typename AttributeType> using TPolygonGroupAttributesRef = TMeshAttrib
 template <typename AttributeType> using TVertexAttributesConstRef = TMeshAttributesConstRef<FVertexID, AttributeType>;
 template <typename AttributeType> using TVertexInstanceAttributesConstRef = TMeshAttributesConstRef<FVertexInstanceID, AttributeType>;
 template <typename AttributeType> using TEdgeAttributesConstRef = TMeshAttributesConstRef<FEdgeID, AttributeType>;
+template <typename AttributeType> using TUVAttributesConstRef = TMeshAttributesConstRef<FUVID, AttributeType>;
 template <typename AttributeType> using TTriangleAttributesConstRef = TMeshAttributesConstRef<FTriangleID, AttributeType>;
 template <typename AttributeType> using TPolygonAttributesConstRef = TMeshAttributesConstRef<FPolygonID, AttributeType>;
 template <typename AttributeType> using TPolygonGroupAttributesConstRef = TMeshAttributesConstRef<FPolygonGroupID, AttributeType>;
@@ -332,9 +88,9 @@ public:
 	// Hence explicitly declare all the below as defaulted, to ensure they will be generated by the compiler.
 	FMeshDescription();
 	~FMeshDescription() = default;
-	FMeshDescription(const FMeshDescription&) = default;
+	FMeshDescription(const FMeshDescription&);
+	FMeshDescription& operator=(const FMeshDescription&);
 	FMeshDescription(FMeshDescription&&) = default;
-	FMeshDescription& operator=(const FMeshDescription&) = default;
 	FMeshDescription& operator=(FMeshDescription&&) = default;
 
 	friend MESHDESCRIPTION_API FArchive& operator<<(FArchive& Ar, FMeshDescription& MeshDescription)
@@ -346,48 +102,81 @@ public:
 	// Serialize the mesh description
 	void Serialize(FArchive& Ar);
 
+	// Legacy serialization for old assets
+	void SerializeLegacy(FArchive& Ar);
+
 	// Empty the meshdescription
 	void Empty();
 
+	// Operations on all the indexers
+	void ResetIndexers();
+	void BuildIndexers();
+	void RebuildIndexers();
+
 	// Return whether the mesh description is empty
 	bool IsEmpty() const;
 
-	FVertexArray& Vertices() { return VertexArray; }
-	const FVertexArray& Vertices() const { return VertexArray; }
+	FVertexArray& Vertices() { return static_cast<TMeshElementContainer<FVertexID>&>(VertexElements->Get()); }
+	const FVertexArray& Vertices() const { return static_cast<const TMeshElementContainer<FVertexID>&>(VertexElements->Get()); }
+
+	FVertexInstanceArray& VertexInstances() { return static_cast<TMeshElementContainer<FVertexInstanceID>&>(VertexInstanceElements->Get()); }
+	const FVertexInstanceArray& VertexInstances() const { return static_cast<const TMeshElementContainer<FVertexInstanceID>&>(VertexInstanceElements->Get()); }
+
+	FEdgeArray& Edges() { return static_cast<TMeshElementContainer<FEdgeID>&>(EdgeElements->Get()); }
+	const FEdgeArray& Edges() const { return static_cast<const TMeshElementContainer<FEdgeID>&>(EdgeElements->Get()); }
+
+	FUVArray& UVs(int32 Index) { return static_cast<TMeshElementContainer<FUVID>&>(UVElements->Get(Index)); }
+	const FUVArray& UVs(int32 Index) const { return static_cast<const TMeshElementContainer<FUVID>&>(UVElements->Get(Index)); }
 
-	FVertexInstanceArray& VertexInstances() { return VertexInstanceArray; }
-	const FVertexInstanceArray& VertexInstances() const { return VertexInstanceArray; }
+	FTriangleArray& Triangles() { return static_cast<TMeshElementContainer<FTriangleID>&>(TriangleElements->Get()); }
+	const FTriangleArray& Triangles() const { return static_cast<const TMeshElementContainer<FTriangleID>&>(TriangleElements->Get()); }
 
-	FEdgeArray& Edges() { return EdgeArray; }
-	const FEdgeArray& Edges() const { return EdgeArray; }
+	FPolygonArray& Polygons() { return static_cast<TMeshElementContainer<FPolygonID>&>(PolygonElements->Get()); }
+	const FPolygonArray& Polygons() const { return static_cast<const TMeshElementContainer<FPolygonID>&>(PolygonElements->Get()); }
 
-	FTriangleArray& Triangles() { return TriangleArray; }
-	const FTriangleArray& Triangles() const { return TriangleArray; }
+	FPolygonGroupArray& PolygonGroups() { return static_cast<TMeshElementContainer<FPolygonGroupID>&>(PolygonGroupElements->Get()); }
+	const FPolygonGroupArray& PolygonGroups() const { return static_cast<const TMeshElementContainer<FPolygonGroupID>&>(PolygonGroupElements->Get()); }
 
-	FPolygonArray& Polygons() { return PolygonArray; }
-	const FPolygonArray& Polygons() const { return PolygonArray; }
+	TAttributesSet<FVertexID>& VertexAttributes() { return Vertices().GetAttributes(); }
+	const TAttributesSet<FVertexID>& VertexAttributes() const { return Vertices().GetAttributes(); }
 
-	FPolygonGroupArray& PolygonGroups() { return PolygonGroupArray; }
-	const FPolygonGroupArray& PolygonGroups() const { return PolygonGroupArray; }
+	TAttributesSet<FVertexInstanceID>& VertexInstanceAttributes() { return VertexInstances().GetAttributes(); }
+	const TAttributesSet<FVertexInstanceID>& VertexInstanceAttributes() const { return VertexInstances().GetAttributes(); }
 
-	TAttributesSet<FVertexID>& VertexAttributes() { return VertexAttributesSet; }
-	const TAttributesSet<FVertexID>& VertexAttributes() const { return VertexAttributesSet; }
+	TAttributesSet<FEdgeID>& EdgeAttributes() { return Edges().GetAttributes(); }
+	const TAttributesSet<FEdgeID>& EdgeAttributes() const { return Edges().GetAttributes(); }
 
-	TAttributesSet<FVertexInstanceID>& VertexInstanceAttributes() { return VertexInstanceAttributesSet; }
-	const TAttributesSet<FVertexInstanceID>& VertexInstanceAttributes() const { return VertexInstanceAttributesSet; }
+	TAttributesSet<FUVID>& UVAttributes(int32 Index) { return UVs(Index).GetAttributes(); }
+	const TAttributesSet<FUVID>& UVAttributes(int32 Index) const { return UVs(Index).GetAttributes(); }
 
-	TAttributesSet<FEdgeID>& EdgeAttributes() { return EdgeAttributesSet; }
-	const TAttributesSet<FEdgeID>& EdgeAttributes() const { return EdgeAttributesSet; }
+	TAttributesSet<FTriangleID>& TriangleAttributes() { return Triangles().GetAttributes(); }
+	const TAttributesSet<FTriangleID>& TriangleAttributes() const { return Triangles().GetAttributes(); }
 
-	TAttributesSet<FTriangleID>& TriangleAttributes() { return TriangleAttributesSet; }
-	const TAttributesSet<FTriangleID>& TriangleAttributes() const { return TriangleAttributesSet; }
+	TAttributesSet<FPolygonID>& PolygonAttributes() { return Polygons().GetAttributes(); }
+	const TAttributesSet<FPolygonID>& PolygonAttributes() const { return Polygons().GetAttributes(); }
 
-	TAttributesSet<FPolygonID>& PolygonAttributes() { return PolygonAttributesSet; }
-	const TAttributesSet<FPolygonID>& PolygonAttributes() const { return PolygonAttributesSet; }
+	TAttributesSet<FPolygonGroupID>& PolygonGroupAttributes() { return PolygonGroups().GetAttributes(); }
+	const TAttributesSet<FPolygonGroupID>& PolygonGroupAttributes() const { return PolygonGroups().GetAttributes(); }
 
-	TAttributesSet<FPolygonGroupID>& PolygonGroupAttributes() { return PolygonGroupAttributesSet; }
-	const TAttributesSet<FPolygonGroupID>& PolygonGroupAttributes() const { return PolygonGroupAttributesSet; }
+	void SuspendVertexIndexing() { VertexToVertexInstances.Suspend(); VertexToEdges.Suspend(); }
+	void SuspendVertexInstanceIndexing() { VertexInstanceToTriangles.Suspend(); }
+	void SuspendEdgeIndexing() { EdgeToTriangles.Suspend(); }
+	void SuspendPolygonIndexing() { PolygonToTriangles.Suspend(); }
+	void SuspendPolygonGroupIndexing() { PolygonGroupToPolygons.Suspend(); PolygonGroupToTriangles.Suspend(); }
+	void SuspendUVIndexing() { UVToTriangles.Suspend(); }
 
+	void ResumeVertexIndexing() { VertexToVertexInstances.Resume(); VertexToEdges.Resume(); }
+	void ResumeVertexInstanceIndexing() { VertexInstanceToTriangles.Resume(); }
+	void ResumeEdgeIndexing() { EdgeToTriangles.Resume(); }
+	void ResumePolygonIndexing() { PolygonToTriangles.Resume(); }
+	void ResumePolygonGroupIndexing() { PolygonGroupToPolygons.Resume(); PolygonGroupToTriangles.Resume(); }
+	void ResumeUVIndexing() { UVToTriangles.Resume(); }
+
+	void BuildVertexIndexers() { VertexToVertexInstances.Build(); VertexToEdges.Build(); }
+	void BuildVertexInstanceIndexers() { VertexInstanceToTriangles.Build(); }
+	void BuildEdgeIndexers() { EdgeToTriangles.Build(); }
+	void BuildPolygonIndexers() { PolygonToTriangles.Build(); }
+	void BuildPolygonGroupIndexers() { PolygonGroupToPolygons.Build(), PolygonGroupToTriangles.Build(); }
 
 //////////////////////////////////////////////////////////////////////////
 // Create / remove mesh elements
@@ -395,49 +184,48 @@ public:
 	/** Reserves space for this number of new vertices */
 	void ReserveNewVertices(const int32 NumVertices)
 	{
-		VertexArray.Reserve(VertexArray.Num() + NumVertices);
+		VertexElements->Get().Reserve(VertexElements->Get().Num() + NumVertices);
 	}
 
 	/** Adds a new vertex to the mesh and returns its ID */
 	FVertexID CreateVertex()
 	{
-		const FVertexID VertexID = VertexArray.Add();
-		CreateVertex_Internal(VertexID);
+		const FVertexID VertexID = VertexElements->Get().Add();
 		return VertexID;
 	}
 
 	/** Adds a new vertex to the mesh with the given ID */
 	void CreateVertexWithID(const FVertexID VertexID)
 	{
-		VertexArray.Insert(VertexID);
-		CreateVertex_Internal(VertexID);
+		VertexElements->Get().Insert(VertexID);
 	}
 
 	/** Deletes a vertex from the mesh */
 	void DeleteVertex(const FVertexID VertexID)
 	{
-		check(VertexArray[VertexID].ConnectedEdgeIDs.Num() == 0);
-		check(VertexArray[VertexID].VertexInstanceIDs.Num() == 0);
-		VertexArray.Remove(VertexID);
-		VertexAttributesSet.Remove(VertexID);
+		check(VertexToVertexInstances.Find(VertexID).Num() == 0);
+		check(VertexToEdges.Find(VertexID).Num() == 0);
+		VertexElements->Get().Remove(VertexID);
+		VertexToVertexInstances.RemoveKey(VertexID);
+		VertexToEdges.RemoveKey(VertexID);
 	}
 
 	/** Returns whether the passed vertex ID is valid */
 	bool IsVertexValid(const FVertexID VertexID) const
 	{
-		return VertexArray.IsValid(VertexID);
+		return VertexElements->Get().IsValid(VertexID);
 	}
 
 	/** Reserves space for this number of new vertex instances */
 	void ReserveNewVertexInstances(const int32 NumVertexInstances)
 	{
-		VertexInstanceArray.Reserve(VertexInstanceArray.Num() + NumVertexInstances);
+		VertexInstanceElements->Get().Reserve(VertexInstanceElements->Get().Num() + NumVertexInstances);
 	}
 
 	/** Adds a new vertex instance to the mesh and returns its ID */
 	FVertexInstanceID CreateVertexInstance(const FVertexID VertexID)
 	{
-		const FVertexInstanceID VertexInstanceID = VertexInstanceArray.Add();
+		const FVertexInstanceID VertexInstanceID = VertexInstanceElements->Get().Add();
 		CreateVertexInstance_Internal(VertexInstanceID, VertexID);
 		return VertexInstanceID;
 	}
@@ -445,7 +233,7 @@ public:
 	/** Adds a new vertex instance to the mesh with the given ID */
 	void CreateVertexInstanceWithID(const FVertexInstanceID VertexInstanceID, const FVertexID VertexID)
 	{
-		VertexInstanceArray.Insert(VertexInstanceID);
+		VertexInstanceElements->Get().Insert(VertexInstanceID);
 		CreateVertexInstance_Internal(VertexInstanceID, VertexID);
 	}
 
@@ -455,19 +243,52 @@ public:
 	/** Returns whether the passed vertex instance ID is valid */
 	bool IsVertexInstanceValid(const FVertexInstanceID VertexInstanceID) const
 	{
-		return VertexInstanceArray.IsValid(VertexInstanceID);
+		return VertexInstanceElements->Get().IsValid(VertexInstanceID);
+	}
+
+	/** Reserves space for this number of new UVs */
+	void ReserveNewUVs(const int32 NumUVs, const int32 UVChannel = 0)
+	{
+		UVElements->Get(UVChannel).Reserve(UVElements->Get(UVChannel).Num() + NumUVs);
+	}
+
+	/** Adds a new UV to the mesh and returns its ID */
+	FUVID CreateUV(const int32 UVChannel = 0)
+	{
+		const FUVID UVID = UVElements->Get(UVChannel).Add();
+		return UVID;
+	}
+
+	/** Adds a new UV to the mesh with the given ID */
+	void CreateUVWithID(const FUVID UVID, const int32 UVChannel = 0)
+	{
+		UVElements->Get(UVChannel).Insert(UVID);
+	}
+
+	/** Deletes a UV from the mesh */
+	void DeleteUV(const FUVID UVID, const int32 UVChannel = 0)
+	{
+		check(UVToTriangles.Find(UVID).Num() == 0);
+		UVElements->Get(UVChannel).Remove(UVID);
+		UVToTriangles.RemoveKey(UVID);
+	}
+
+	/** Returns whether the passed UV ID is valid */
+	bool IsUVValid(const FUVID UVID, const int32 UVChannel = 0) const
+	{
+		return UVElements->Get(UVChannel).IsValid(UVID);
 	}
 
 	/** Reserves space for this number of new edges */
 	void ReserveNewEdges(const int32 NumEdges)
 	{
-		EdgeArray.Reserve(EdgeArray.Num() + NumEdges);
+		EdgeElements->Get().Reserve(EdgeElements->Get().Num() + NumEdges);
 	}
 
 	/** Adds a new edge to the mesh and returns its ID */
 	FEdgeID CreateEdge(const FVertexID VertexID0, const FVertexID VertexID1)
 	{
-		const FEdgeID EdgeID = EdgeArray.Add();
+		const FEdgeID EdgeID = EdgeElements->Get().Add();
 		CreateEdge_Internal(EdgeID, VertexID0, VertexID1);
 		return EdgeID;
 	}
@@ -475,7 +296,7 @@ public:
 	/** Adds a new edge to the mesh with the given ID */
 	void CreateEdgeWithID(const FEdgeID EdgeID, const FVertexID VertexID0, const FVertexID VertexID1)
 	{
-		EdgeArray.Insert(EdgeID);
+		EdgeElements->Get().Insert(EdgeID);
 		CreateEdge_Internal(EdgeID, VertexID0, VertexID1);
 	}
 
@@ -485,19 +306,19 @@ public:
 	/** Returns whether the passed edge ID is valid */
 	bool IsEdgeValid(const FEdgeID EdgeID) const
 	{
-		return EdgeArray.IsValid(EdgeID);
+		return EdgeElements->Get().IsValid(EdgeID);
 	}
 
 	/** Reserves space for this number of new triangles */
 	void ReserveNewTriangles(const int32 NumTriangles)
 	{
-		TriangleArray.Reserve(TriangleArray.Num() + NumTriangles);
+		TriangleElements->Get().Reserve(TriangleElements->Get().Num() + NumTriangles);
 	}
 
 	/** Adds a new triangle to the mesh and returns its ID. This will also make an encapsulating polygon, and any missing edges. */
 	FTriangleID CreateTriangle(const FPolygonGroupID PolygonGroupID, TArrayView<const FVertexInstanceID> VertexInstanceIDs, TArray<FEdgeID>* OutEdgeIDs = nullptr)
 	{
-		const FTriangleID TriangleID = TriangleArray.Add();
+		const FTriangleID TriangleID = TriangleElements->Get().Add();
 		CreateTriangle_Internal(TriangleID, PolygonGroupID, VertexInstanceIDs, OutEdgeIDs);
 		return TriangleID;
 	}
@@ -505,7 +326,7 @@ public:
 	/** Adds a new triangle to the mesh with the given ID. This will also make an encapsulating polygon, and any missing edges. */
 	void CreateTriangleWithID(const FTriangleID TriangleID, const FPolygonGroupID PolygonGroupID, TArrayView<const FVertexInstanceID> VertexInstanceIDs, TArray<FEdgeID>* OutEdgeIDs = nullptr)
 	{
-		TriangleArray.Insert(TriangleID);
+		TriangleElements->Get().Insert(TriangleID);
 		CreateTriangle_Internal(TriangleID, PolygonGroupID, VertexInstanceIDs, OutEdgeIDs);
 	}
 
@@ -520,19 +341,19 @@ public:
 	/** Returns whether the passed triangle ID is valid */
 	bool IsTriangleValid(const FTriangleID TriangleID) const
 	{
-		return TriangleArray.IsValid(TriangleID);
+		return TriangleElements->Get().IsValid(TriangleID);
 	}
 
 	/** Reserves space for this number of new polygons */
 	void ReserveNewPolygons(const int32 NumPolygons)
 	{
-		PolygonArray.Reserve(PolygonArray.Num() + NumPolygons);
+		PolygonElements->Get().Reserve(PolygonElements->Get().Num() + NumPolygons);
 	}
 
 	/** Adds a new polygon to the mesh and returns its ID. This will also make any missing edges, and all constituent triangles. */
 	FPolygonID CreatePolygon(const FPolygonGroupID PolygonGroupID, TArrayView<const FVertexInstanceID> VertexInstanceIDs, TArray<FEdgeID>* OutEdgeIDs = nullptr)
 	{
-		const FPolygonID PolygonID = PolygonArray.Add();
+		const FPolygonID PolygonID = PolygonElements->Get().Add();
 		CreatePolygon_Internal(PolygonID, PolygonGroupID, VertexInstanceIDs, OutEdgeIDs);
 		return PolygonID;
 	}
@@ -540,7 +361,7 @@ public:
 	/** Adds a new polygon to the mesh with the given ID. This will also make any missing edges, and all constituent triangles. */
 	void CreatePolygonWithID(const FPolygonID PolygonID, const FPolygonGroupID PolygonGroupID, TArrayView<const FVertexInstanceID> VertexInstanceIDs, TArray<FEdgeID>* OutEdgeIDs = nullptr)
 	{
-		PolygonArray.Insert(PolygonID);
+		PolygonElements->Get().Insert(PolygonID);
 		CreatePolygon_Internal(PolygonID, PolygonGroupID, VertexInstanceIDs, OutEdgeIDs);
 	}
 
@@ -555,42 +376,41 @@ public:
 	/** Returns whether the passed polygon ID is valid */
 	bool IsPolygonValid(const FPolygonID PolygonID) const
 	{
-		return PolygonArray.IsValid(PolygonID);
+		return PolygonElements->Get().IsValid(PolygonID);
 	}
 
 	/** Reserves space for this number of new polygon groups */
 	void ReserveNewPolygonGroups(const int32 NumPolygonGroups)
 	{
-		PolygonGroupArray.Reserve(PolygonGroupArray.Num() + NumPolygonGroups);
+		PolygonGroupElements->Get().Reserve(PolygonGroupElements->Get().Num() + NumPolygonGroups);
 	}
 
 	/** Adds a new polygon group to the mesh and returns its ID */
 	FPolygonGroupID CreatePolygonGroup()
 	{
-		const FPolygonGroupID PolygonGroupID = PolygonGroupArray.Add();
-		CreatePolygonGroup_Internal(PolygonGroupID);
+		const FPolygonGroupID PolygonGroupID = PolygonGroupElements->Get().Add();
 		return PolygonGroupID;
 	}
 
 	/** Adds a new polygon group to the mesh with the given ID */
 	void CreatePolygonGroupWithID(const FPolygonGroupID PolygonGroupID)
 	{
-		PolygonGroupArray.Insert(PolygonGroupID);
-		CreatePolygonGroup_Internal(PolygonGroupID);
+		PolygonGroupElements->Get().Insert(PolygonGroupID);
 	}
 
 	/** Deletes a polygon group from the mesh */
 	void DeletePolygonGroup(const FPolygonGroupID PolygonGroupID)
 	{
-		check(PolygonGroupArray[PolygonGroupID].Polygons.Num() == 0);
-		PolygonGroupArray.Remove(PolygonGroupID);
-		PolygonGroupAttributesSet.Remove(PolygonGroupID);
+		check(PolygonGroupToPolygons.Find(PolygonGroupID).Num() == 0);
+		PolygonGroupElements->Get().Remove(PolygonGroupID);
+		PolygonGroupToPolygons.RemoveKey(PolygonGroupID);
+		PolygonGroupToTriangles.RemoveKey(PolygonGroupID);
 	}
 
 	/** Returns whether the passed polygon group ID is valid */
 	bool IsPolygonGroupValid(const FPolygonGroupID PolygonGroupID) const
 	{
-		return PolygonGroupArray.IsValid(PolygonGroupID);
+		return PolygonGroupElements->Get().IsValid(PolygonGroupID);
 	}
 
 
@@ -605,31 +425,31 @@ public:
 	/** Returns whether a given vertex is orphaned, i.e. it doesn't form part of any polygon */
 	bool IsVertexOrphaned(const FVertexID VertexID) const;
 
-	/** Returns the edge ID defined by the two given vertex IDs, if there is one; otherwise FEdgeID::Invalid */
+	/** Returns the edge ID defined by the two given vertex IDs, if there is one; otherwise INDEX_NONE */
 	FEdgeID GetVertexPairEdge(const FVertexID VertexID0, const FVertexID VertexID1) const;
 
 	/** Returns reference to an array of Edge IDs connected to this vertex */
-	const TArray<FEdgeID>& GetVertexConnectedEdges(const FVertexID VertexID) const
+	TArrayView<const FEdgeID> GetVertexConnectedEdges(const FVertexID VertexID) const
 	{
-		return VertexArray[VertexID].ConnectedEdgeIDs;
+		return VertexToEdges.Find<FEdgeID>(VertexID);
 	}
 
 	/** Returns number of edges connected to this vertex */
 	int32 GetNumVertexConnectedEdges(const FVertexID VertexID) const
 	{
-		return VertexArray[VertexID].ConnectedEdgeIDs.Num();
+		return VertexToEdges.Find(VertexID).Num();
 	}
 
 	/** Returns reference to an array of VertexInstance IDs instanced from this vertex */
-	const TArray<FVertexInstanceID>& GetVertexVertexInstances(const FVertexID VertexID) const
+	TArrayView<const FVertexInstanceID> GetVertexVertexInstances(const FVertexID VertexID) const
 	{
-		return VertexArray[VertexID].VertexInstanceIDs;
+		return VertexToVertexInstances.Find<FVertexInstanceID>(VertexID);
 	}
 
 	/** Returns number of vertex instances created from this vertex */
 	int32 GetNumVertexVertexInstances(const FVertexID VertexID) const
 	{
-		return VertexArray[VertexID].VertexInstanceIDs.Num();
+		return VertexToVertexInstances.Find(VertexID).Num();
 	}
 
 	/** Populates the passed array of TriangleIDs with the triangles connected to this vertex */
@@ -637,9 +457,10 @@ public:
 	void GetVertexConnectedTriangles(const FVertexID VertexID, TArray<FTriangleID, Alloc>& OutConnectedTriangleIDs) const
 	{
 		OutConnectedTriangleIDs.Reset(GetNumVertexConnectedTriangles(VertexID));
-		for (const FVertexInstanceID VertexInstanceID : VertexArray[VertexID].VertexInstanceIDs)
+		for (const FVertexInstanceID VertexInstanceID : VertexToVertexInstances.Find<FVertexInstanceID>(VertexID))
 		{
-			OutConnectedTriangleIDs.Append(VertexInstanceArray[VertexInstanceID].ConnectedTriangles);
+			TArrayView<const FTriangleID> ConnectedTris = VertexInstanceToTriangles.Find<FTriangleID>(VertexInstanceID);
+			OutConnectedTriangleIDs.Append(ConnectedTris.GetData(), ConnectedTris.Num());
 		}
 	}
 
@@ -663,8 +484,9 @@ public:
 	/** Returns number of triangles connected to this vertex */
 	int32 GetNumVertexConnectedTriangles(const FVertexID VertexID) const
 	{
-		return Algo::TransformAccumulate(VertexArray[VertexID].VertexInstanceIDs,
-										 [this](const FVertexInstanceID ID) { return VertexInstanceArray[ID].ConnectedTriangles.Num(); },
+		TArrayView<const FVertexInstanceID> VertexInstances = VertexToVertexInstances.Find<FVertexInstanceID>(VertexID);
+		return Algo::TransformAccumulate(VertexInstances,
+										 [this](const FVertexInstanceID ID) { return VertexInstanceToTriangles.Find(ID).Num(); },
 										 0);
 	}
 
@@ -673,11 +495,11 @@ public:
 	void GetVertexConnectedPolygons(const FVertexID VertexID, TArray<FPolygonID, Alloc>& OutConnectedPolygonIDs) const
 	{
 		OutConnectedPolygonIDs.Reset();
-		for (const FVertexInstanceID VertexInstanceID : VertexArray[VertexID].VertexInstanceIDs)
+		for (const FVertexInstanceID VertexInstanceID : VertexToVertexInstances.Find<FVertexInstanceID>(VertexID))
 		{
-			for (const FTriangleID TriangleID : VertexInstanceArray[VertexInstanceID].ConnectedTriangles)
+			for (const FTriangleID TriangleID : VertexInstanceToTriangles.Find<FTriangleID>(VertexInstanceID))
 			{
-				OutConnectedPolygonIDs.AddUnique(TriangleArray[TriangleID].PolygonID);
+				OutConnectedPolygonIDs.AddUnique(TrianglePolygons[TriangleID]);
 			}
 		}
 	}
@@ -709,14 +531,14 @@ public:
 	template <typename Alloc>
 	void GetVertexAdjacentVertices(const FVertexID VertexID, TArray<FVertexID, Alloc>& OutAdjacentVertexIDs) const
 	{
-		const TArray<FEdgeID>& ConnectedEdgeIDs = VertexArray[VertexID].ConnectedEdgeIDs;
+		TArrayView<const FEdgeID> ConnectedEdgeIDs = VertexToEdges.Find<FEdgeID>(VertexID);
 		OutAdjacentVertexIDs.SetNumUninitialized(ConnectedEdgeIDs.Num());
 
 		int32 Index = 0;
 		for (const FEdgeID EdgeID : ConnectedEdgeIDs)
 		{
-			const FMeshEdge& Edge = EdgeArray[EdgeID];
-			OutAdjacentVertexIDs[Index] = (Edge.VertexIDs[0] == VertexID) ? Edge.VertexIDs[1] : Edge.VertexIDs[0];
+			TArrayView<const FVertexID> EdgeVertexIDs = EdgeVertices[EdgeID];
+			OutAdjacentVertexIDs[Index] = (EdgeVertexIDs[0] == VertexID) ? EdgeVertexIDs[1] : EdgeVertexIDs[0];
 			Index++;
 		}
 	}
@@ -738,6 +560,16 @@ public:
 		return Result;
 	}
 
+	/** Returns the position of this vertex */
+	FVector GetVertexPosition(const FVertexID VertexID) const
+	{
+		return VertexPositions[VertexID];
+	}
+
+	TMeshAttributesArray<FVector> GetVertexPositions() const
+	{
+		return VertexPositions;
+	}
 
 	//////////////////////////////////////////////////////////////////////
 	// Vertex instance operations
@@ -745,32 +577,32 @@ public:
 	/** Returns the vertex ID associated with the given vertex instance */
 	FVertexID GetVertexInstanceVertex(const FVertexInstanceID VertexInstanceID) const
 	{
-		return VertexInstanceArray[VertexInstanceID].VertexID;
+		return VertexInstanceVertices[VertexInstanceID];
 	}
 
-	/** Returns the edge ID defined by the two given vertex instance IDs, if there is one; otherwise FEdgeID::Invalid */
+	/** Returns the edge ID defined by the two given vertex instance IDs, if there is one; otherwise INDEX_NONE */
 	FEdgeID GetVertexInstancePairEdge(const FVertexInstanceID VertexInstanceID0, const FVertexInstanceID VertexInstanceID1) const;
 
 	/** Returns reference to an array of Triangle IDs connected to this vertex instance */
-	const TArray<FTriangleID>& GetVertexInstanceConnectedTriangles(const FVertexInstanceID VertexInstanceID) const
+	TArrayView<const FTriangleID> GetVertexInstanceConnectedTriangles(const FVertexInstanceID VertexInstanceID) const
 	{
-		return VertexInstanceArray[VertexInstanceID].ConnectedTriangles;
+		return VertexInstanceToTriangles.Find<FTriangleID>(VertexInstanceID);
 	}
 
 	/** Returns the number of triangles connected to this vertex instance */
 	int32 GetNumVertexInstanceConnectedTriangles(const FVertexInstanceID VertexInstanceID) const
 	{
-		return VertexInstanceArray[VertexInstanceID].ConnectedTriangles.Num();
+		return VertexInstanceToTriangles.Find(VertexInstanceID).Num();
 	}
 
 	/** Populates the passed array with the polygons connected to this vertex instance */
 	template <typename Alloc>
 	void GetVertexInstanceConnectedPolygons(const FVertexInstanceID VertexInstanceID, TArray<FPolygonID, Alloc>& OutPolygonIDs) const
 	{
-		OutPolygonIDs.Reset(VertexInstanceArray[VertexInstanceID].ConnectedTriangles.Num());
-		for (const FTriangleID TriangleID : VertexInstanceArray[VertexInstanceID].ConnectedTriangles)
+		OutPolygonIDs.Reset(VertexInstanceToTriangles.Find(VertexInstanceID).Num());
+		for (const FTriangleID TriangleID : VertexInstanceToTriangles.Find<FTriangleID>(VertexInstanceID))
 		{
-			OutPolygonIDs.AddUnique(TriangleArray[TriangleID].PolygonID);
+			OutPolygonIDs.AddUnique(TrianglePolygons[TriangleID]);
 		}
 	}
 
@@ -804,39 +636,39 @@ public:
 	/** Determine whether a given edge is an internal edge between triangles of a polygon */
 	bool IsEdgeInternal(const FEdgeID EdgeID) const
 	{
-		const TArray<FTriangleID>& ConnectedTriangles = EdgeArray[EdgeID].ConnectedTriangles;
+		TArrayView<const FTriangleID> ConnectedTriangles = EdgeToTriangles.Find<FTriangleID>(EdgeID);
 		return ConnectedTriangles.Num() == 2 &&
-			   TriangleArray[ConnectedTriangles[0]].PolygonID == TriangleArray[ConnectedTriangles[1]].PolygonID;
+			   TrianglePolygons[ConnectedTriangles[0]] == TrianglePolygons[ConnectedTriangles[1]];
 	}
 
 	/** Determine whether a given edge is an internal edge between triangles of a specific polygon */
 	bool IsEdgeInternalToPolygon(const FEdgeID EdgeID, const FPolygonID PolygonID) const
 	{
-		const TArray<FTriangleID>& ConnectedTriangles = EdgeArray[EdgeID].ConnectedTriangles;
+		TArrayView<const FTriangleID> ConnectedTriangles = EdgeToTriangles.Find<FTriangleID>(EdgeID);
 		return ConnectedTriangles.Num() == 2 &&
-			   TriangleArray[ConnectedTriangles[0]].PolygonID == PolygonID &&
-			   TriangleArray[ConnectedTriangles[1]].PolygonID == PolygonID;
+			   TrianglePolygons[ConnectedTriangles[0]] == PolygonID &&
+			   TrianglePolygons[ConnectedTriangles[1]] == PolygonID;
 	}
 
 	/** Returns reference to an array of triangle IDs connected to this edge */
-	const TArray<FTriangleID>& GetEdgeConnectedTriangles(const FEdgeID EdgeID) const
+	TArrayView<const FTriangleID> GetEdgeConnectedTriangles(const FEdgeID EdgeID) const
 	{
-		return EdgeArray[EdgeID].ConnectedTriangles;
+		return EdgeToTriangles.Find<FTriangleID>(EdgeID);
 	}
 
 	int32 GetNumEdgeConnectedTriangles(const FEdgeID EdgeID) const
 	{
-		return EdgeArray[EdgeID].ConnectedTriangles.Num();
+		return EdgeToTriangles.Find(EdgeID).Num();
 	}
 
 	/** Populates the passed array with polygon IDs connected to this edge */
 	template <typename Alloc>
 	void GetEdgeConnectedPolygons(const FEdgeID EdgeID, TArray<FPolygonID, Alloc>& OutPolygonIDs) const
 	{
-		OutPolygonIDs.Reset(EdgeArray[EdgeID].ConnectedTriangles.Num());
-		for (const FTriangleID TriangleID : EdgeArray[EdgeID].ConnectedTriangles)
+		OutPolygonIDs.Reset(EdgeToTriangles.Find(EdgeID).Num());
+		for (const FTriangleID TriangleID : EdgeToTriangles.Find<FTriangleID>(EdgeID))
 		{
-			OutPolygonIDs.AddUnique(TriangleArray[TriangleID].PolygonID);
+			OutPolygonIDs.AddUnique(TrianglePolygons[TriangleID]);
 		}
 	}
 
@@ -867,13 +699,14 @@ public:
 	FVertexID GetEdgeVertex(const FEdgeID EdgeID, int32 VertexNumber) const
 	{
 		check(VertexNumber == 0 || VertexNumber == 1);
-		return EdgeArray[EdgeID].VertexIDs[VertexNumber];
+		TArrayView<const FVertexID> EdgeVertexIDs = EdgeVertices[EdgeID];
+		return EdgeVertexIDs[VertexNumber];
 	}
 
 	/** Returns a pair of vertex IDs defining the edge */
 	TArrayView<const FVertexID> GetEdgeVertices(const FEdgeID EdgeID) const
 	{
-		return EdgeArray[EdgeID].VertexIDs;
+		return EdgeVertices[EdgeID];
 	}
 
 
@@ -883,68 +716,69 @@ public:
 	/** Get the polygon which contains this triangle */
 	FPolygonID GetTrianglePolygon(const FTriangleID TriangleID) const
 	{
-		return TriangleArray[TriangleID].PolygonID;
+		return TrianglePolygons[TriangleID];
 	}
 
 	/** Get the polygon group which contains this triangle */
 	FPolygonGroupID GetTrianglePolygonGroup(const FTriangleID TriangleID) const
 	{
-		return PolygonArray[TriangleArray[TriangleID].PolygonID].PolygonGroupID;
+		return TrianglePolygonGroups[TriangleID];
 	}
 
 	/** Determines if this triangle is part of an n-gon */
 	bool IsTrianglePartOfNgon(const FTriangleID TriangleID) const
 	{
-		return PolygonArray[TriangleArray[TriangleID].PolygonID].TriangleIDs.Num() > 1;
+		return PolygonToTriangles.Find(TrianglePolygons[TriangleID]).Num() > 1;
 	}
 
 	/** Get the vertex instances which define this triangle */
 	TArrayView<const FVertexInstanceID> GetTriangleVertexInstances(const FTriangleID TriangleID) const
 	{
-		return TriangleArray[TriangleID].VertexInstanceIDs;
+		return TriangleVertexInstances[TriangleID];
 	}
 
 	/** Get the specified vertex instance by index */
 	FVertexInstanceID GetTriangleVertexInstance(const FTriangleID TriangleID, const int32 Index) const
 	{
 		check(Index >= 0 && Index < 3);
-		return TriangleArray[TriangleID].GetVertexInstanceID(Index);
+		TArrayView<const FVertexInstanceID> TriVertexInstanceIDs = TriangleVertexInstances[TriangleID];
+		return TriVertexInstanceIDs[Index];
 	}
 
 	/** Populates the passed array with the vertices which define this triangle */
+	UE_DEPRECATED(4.26, "Use the other form of GetTriangleVertices")
 	void GetTriangleVertices(const FTriangleID TriangleID, TArrayView<FVertexID> OutVertexIDs) const
 	{
 		check(OutVertexIDs.Num() >= 3);
+		TArrayView<const FVertexID> TriVerts = TriangleVertices[TriangleID];
 		for (int32 Index = 0; Index < 3; ++Index)
 		{
-			OutVertexIDs[Index] = GetVertexInstanceVertex(TriangleArray[TriangleID].GetVertexInstanceID(Index));
+			OutVertexIDs[Index] = TriVerts[Index];
 		}
 	}
 
 	/** Return the vertices which define this triangle */
-	TStaticArray<FVertexID, 3> GetTriangleVertices(const FTriangleID TriangleID) const
+	TArrayView<const FVertexID> GetTriangleVertices(const FTriangleID TriangleID) const
 	{
-		TStaticArray<FVertexID, 3> Result;
-		GetTriangleVertices(TriangleID, Result);
-		return Result;
+		return TriangleVertices[TriangleID];
 	}
 
 	/** Populates the passed array with the edges which define this triangle */
+	UE_DEPRECATED(4.26, "Use the other form of GetTriangleEdges")
 	void GetTriangleEdges(const FTriangleID TriangleID, TArrayView<FEdgeID> OutEdgeIDs) const
 	{
 		check(OutEdgeIDs.Num() >= 3);
-		TStaticArray<FVertexID, 3> VertexIDs = GetTriangleVertices(TriangleID);
-		OutEdgeIDs[0] = GetVertexPairEdge(VertexIDs[0], VertexIDs[1]);
-		OutEdgeIDs[1] = GetVertexPairEdge(VertexIDs[1], VertexIDs[2]);
-		OutEdgeIDs[2] = GetVertexPairEdge(VertexIDs[2], VertexIDs[0]);
+		TArrayView<const FEdgeID> TriEdges = TriangleEdges[TriangleID];
+		for (int32 Index = 0; Index < 3; ++Index)
+		{
+			OutEdgeIDs[Index] = TriEdges[Index];
+		}
 	}
 
 	/** Return the edges which form this triangle */
-	TStaticArray<FEdgeID, 3> GetTriangleEdges(const FTriangleID TriangleID) const
+	TArrayView<FEdgeID> GetTriangleEdges(const FTriangleID TriangleID) const
 	{
-		TStaticArray<FEdgeID, 3> Result;
-		GetTriangleEdges(TriangleID, Result);
-		return Result;
+		return TriangleEdges[TriangleID];
 	}
 
 	/** Populates the passed array with adjacent triangles */
@@ -954,7 +788,7 @@ public:
 		OutTriangleIDs.Reset();
 		for (const FEdgeID EdgeID : GetTriangleEdges(TriangleID))
 		{
-			for (const FTriangleID OtherTriangleID : EdgeArray[EdgeID].ConnectedTriangles)
+			for (const FTriangleID OtherTriangleID : EdgeToTriangles.Find(EdgeID))
 			{
 				if (OtherTriangleID != TriangleID)
 				{
@@ -981,42 +815,87 @@ public:
 		return Result;
 	}
 
-	/** Return the vertex instance which corresponds to the given vertex on the given triangle, or FVertexInstanceID::Invalid */
+	/** Return UV indices for this triangle for the given channel */
+	TArrayView<FUVID> GetTriangleUVIndices(const FTriangleID TriangleID, int32 UVChannel = 0)
+	{
+		return TriangleUVs.Get(TriangleID, UVChannel);
+	}
+
+	/** Return the vertex instance which corresponds to the given vertex on the given triangle, or INDEX_NONE */
 	FVertexInstanceID GetVertexInstanceForTriangleVertex(const FTriangleID TriangleID, const FVertexID VertexID) const
 	{
 		const FVertexInstanceID* VertexInstanceIDPtr = Algo::FindByPredicate(
 			GetTriangleVertexInstances(TriangleID),
 			[this, VertexID](const FVertexInstanceID VertexInstanceID) { return (GetVertexInstanceVertex(VertexInstanceID) == VertexID); });
 
-		return VertexInstanceIDPtr ? *VertexInstanceIDPtr : FVertexInstanceID::Invalid;
+		return VertexInstanceIDPtr ? *VertexInstanceIDPtr : FVertexInstanceID(INDEX_NONE);
 	}
 
+	/** Reverse the winding order of the vertices of this triangle */
+	void ReverseTriangleFacing(const FTriangleID TriangleID);
+
+	/** Set the UV indices for this triangle */
+	void SetTriangleUVIndices(const FTriangleID TriangleID, TArrayView<const FUVID> UVIDs, int32 UVChannel = 0);
 
 	//////////////////////////////////////////////////////////////////////
 	// Polygon operations
 
 	/** Return reference to an array of triangle IDs which comprise this polygon */
-	const TArray<FTriangleID>& GetPolygonTriangleIDs(const FPolygonID PolygonID) const
+	TArrayView<const FTriangleID> GetPolygonTriangleIDs(const FPolygonID PolygonID) const
 	{
-		return PolygonArray[PolygonID].TriangleIDs;
+		return PolygonToTriangles.Find<FTriangleID>(PolygonID);
 	}
 
 	/** Return the number of triangles which comprise this polygon */
 	int32 GetNumPolygonTriangles(const FPolygonID PolygonID) const
 	{
-		return PolygonArray[PolygonID].TriangleIDs.Num();
+		return PolygonToTriangles.Find(PolygonID).Num();
 	}
 
 	/** Returns reference to an array of VertexInstance IDs forming the perimeter of this polygon */
-	const TArray<FVertexInstanceID>& GetPolygonVertexInstances(const FPolygonID PolygonID) const
+	template <typename Alloc>
+	void GetPolygonVertexInstances(const FPolygonID PolygonID, TArray<FVertexInstanceID, Alloc>& OutVertexInstanceIDs) const
 	{
-		return PolygonArray[PolygonID].VertexInstanceIDs;
+		OutVertexInstanceIDs.SetNumUninitialized(GetNumPolygonVertices(PolygonID));
+		TArrayView<const FTriangleID> Tris = PolygonToTriangles.Find<FTriangleID>(PolygonID);
+		if (Tris.Num() == 1)
+		{
+			OutVertexInstanceIDs = TriangleVertexInstances[Tris[0]];
+		}
+		else
+		{
+			TArray<TTuple<int32, int32>, TInlineAllocator<8>> Result;
+			Result.SetNumUninitialized(Tris.Num() + 2);
+			FindPolygonPerimeter(Tris, Result);
+			for (int32 Index = 0; Index < Result.Num(); Index++)
+			{
+				FTriangleID TriangleID = Tris[Result[Index].Get<0>()];
+				int32 VertexInstanceNumber = Result[Index].Get<1>();
+				OutVertexInstanceIDs[Index] = TriangleVertexInstances[TriangleID][VertexInstanceNumber];
+			}
+		}
+	}
+
+	template <typename Alloc>
+	TArray<FVertexInstanceID, Alloc> GetPolygonVertexInstances(const FPolygonID PolygonID) const
+	{
+		TArray<FVertexInstanceID, Alloc> Result;
+		this->GetPolygonVertexInstances(PolygonID, Result);
+		return Result;
+	}
+
+	TArray<FVertexInstanceID> GetPolygonVertexInstances(const FPolygonID PolygonID) const
+	{
+		TArray<FVertexInstanceID> Result;
+		this->GetPolygonVertexInstances(PolygonID, Result);
+		return Result;
 	}
 
 	/** Returns the number of vertices this polygon has */
 	int32 GetNumPolygonVertices(const FPolygonID PolygonID) const
 	{
-		return PolygonArray[PolygonID].VertexInstanceIDs.Num();
+		TArrayView<const FTriangleID> Tris = PolygonToTriangles.Find<FTriangleID>(PolygonID);
+		return Tris.Num() + 2;
 	}
 
 	/** Populates the passed array of VertexIDs with the vertices which form the polygon perimeter */
@@ -1024,10 +903,22 @@ public:
 	void GetPolygonVertices(const FPolygonID PolygonID, TArray<FVertexID, Alloc>& OutVertexIDs) const
 	{
 		OutVertexIDs.SetNumUninitialized(GetNumPolygonVertices(PolygonID));
-		int32 Index = 0;
-		for (const FVertexInstanceID VertexInstanceID : GetPolygonVertexInstances(PolygonID))
+		TArrayView<const FTriangleID> Tris = PolygonToTriangles.Find<FTriangleID>(PolygonID);
+		if (Tris.Num() == 1)
 		{
-			OutVertexIDs[Index++] = GetVertexInstanceVertex(VertexInstanceID);
+			OutVertexIDs = TriangleVertices[Tris[0]];
+		}
+		else
+		{
+			TArray<TTuple<int32, int32>, TInlineAllocator<8>> Result;
+			Result.SetNumUninitialized(Tris.Num() + 2);
+			FindPolygonPerimeter(Tris, Result);
+			for (int32 Index = 0; Index < Result.Num(); Index++)
+			{
+				FTriangleID TriangleID = Tris[Result[Index].Get<0>()];
+				int32 VertexNumber = Result[Index].Get<1>();
+				OutVertexIDs[Index] = TriangleVertices[TriangleID][VertexNumber];
+			}
 		}
 	}
 
@@ -1052,14 +943,23 @@ public:
 	template <typename Alloc>
 	void GetPolygonPerimeterEdges(const FPolygonID PolygonID, TArray<FEdgeID, Alloc>& OutEdgeIDs) const
 	{
-		const TArray<FVertexInstanceID>& VertexInstanceIDs = GetPolygonVertexInstances(PolygonID);
-		const int32 ContourCount = VertexInstanceIDs.Num();
-		OutEdgeIDs.SetNumUninitialized(ContourCount);
-		for (int32 ContourIndex = 0; ContourIndex < ContourCount; ++ContourIndex)
+		OutEdgeIDs.SetNumUninitialized(GetNumPolygonVertices(PolygonID));
+		TArrayView<const FTriangleID> Tris = PolygonToTriangles.Find<FTriangleID>(PolygonID);
+		if (Tris.Num() == 1)
 		{
-			const int32 ContourPlusOne = (ContourIndex == ContourCount - 1) ? 0 : ContourIndex + 1;
-			OutEdgeIDs[ContourIndex] = GetVertexPairEdge(GetVertexInstanceVertex(VertexInstanceIDs[ContourIndex]),
-														 GetVertexInstanceVertex(VertexInstanceIDs[ContourPlusOne]));
+			OutEdgeIDs = TriangleEdges[Tris[0]];
+		}
+		else
+		{
+			TArray<TTuple<int32, int32>, TInlineAllocator<8>> Result;
+			Result.SetNumUninitialized(Tris.Num() + 2);
+			FindPolygonPerimeter(Tris, Result);
+			for (int32 Index = 0; Index < Result.Num(); Index++)
+			{
+				FTriangleID TriangleID = Tris[Result[Index].Get<0>()];
+				int32 EdgeNumber = Result[Index].Get<1>();
+				OutEdgeIDs[Index] = TriangleEdges[TriangleID][EdgeNumber];
+			}
 		}
 	}
 
@@ -1088,7 +988,8 @@ public:
 		OutEdgeIDs.Reset(GetNumPolygonVertices(PolygonID) - 3);
 		if (GetNumPolygonVertices(PolygonID) > 3)
 		{
-			for (const FVertexInstanceID VertexInstanceID : GetPolygonVertexInstances(PolygonID))
+			TArray<FVertexInstanceID> VertexInstanceIDs = GetPolygonVertexInstances(PolygonID);
+			for (const FVertexInstanceID VertexInstanceID : VertexInstanceIDs)
 			{
 				for (const FEdgeID EdgeID : GetVertexConnectedEdges(GetVertexInstanceVertex(VertexInstanceID)))
 				{
@@ -1121,7 +1022,7 @@ public:
 	/** Return the number of internal edges in this polygon */
 	int32 GetNumPolygonInternalEdges(const FPolygonID PolygonID) const
 	{
-		return PolygonArray[PolygonID].VertexInstanceIDs.Num() - 3;
+		return PolygonToTriangles.Find(PolygonID).Num() - 1;
 	}
 
 	/** Populates the passed array with adjacent polygons */
@@ -1161,34 +1062,49 @@ public:
 	/** Return the polygon group associated with a polygon */
 	FPolygonGroupID GetPolygonPolygonGroup(const FPolygonID PolygonID) const
 	{
-		return PolygonArray[PolygonID].PolygonGroupID;
+		return PolygonPolygonGroups[PolygonID];
 	}
 
-	/** Return the vertex instance which corresponds to the given vertex on the given polygon, or FVertexInstanceID::Invalid */
+	/** Return the vertex instance which corresponds to the given vertex on the given polygon, or INDEX_NONE */
 	FVertexInstanceID GetVertexInstanceForPolygonVertex(const FPolygonID PolygonID, const FVertexID VertexID) const
 	{
-		const FVertexInstanceID* VertexInstanceIDPtr = GetPolygonVertexInstances(PolygonID).FindByPredicate(
+		TArray<FVertexInstanceID> VertexInstanceIDs = GetPolygonVertexInstances(PolygonID);
+		const FVertexInstanceID* VertexInstanceIDPtr = VertexInstanceIDs.FindByPredicate(
 			[this, VertexID](const FVertexInstanceID VertexInstanceID) { return (GetVertexInstanceVertex(VertexInstanceID) == VertexID); });
 
-		return VertexInstanceIDPtr ? *VertexInstanceIDPtr : FVertexInstanceID::Invalid;
+		return VertexInstanceIDPtr ? *VertexInstanceIDPtr : FVertexInstanceID(INDEX_NONE);
 	}
 
 	/** Set the vertex instance at the given index around the polygon to the new value */
 	void SetPolygonVertexInstance(const FPolygonID PolygonID, const int32 PerimeterIndex, const FVertexInstanceID VertexInstanceID);
+	void SetPolygonVertexInstances(const FPolygonID PolygonID, TArrayView<const FVertexInstanceID> VertexInstanceIDs);
 
 	/** Sets the polygon group associated with a polygon */
 	void SetPolygonPolygonGroup(const FPolygonID PolygonID, const FPolygonGroupID PolygonGroupID)
 	{
-		FMeshPolygon& Polygon = PolygonArray[PolygonID];
-		verify(PolygonGroupArray[Polygon.PolygonGroupID].Polygons.Remove(PolygonID) == 1);
-		Polygon.PolygonGroupID = PolygonGroupID;
-		check(!PolygonGroupArray[PolygonGroupID].Polygons.Contains(PolygonID));
-		PolygonGroupArray[PolygonGroupID].Polygons.Add(PolygonID);
+		FPolygonGroupID OldPolygonGroupID = PolygonPolygonGroups[PolygonID];
+		PolygonGroupToPolygons.RemoveReferenceFromKey(OldPolygonGroupID, PolygonID);
+
+		PolygonPolygonGroups[PolygonID] = PolygonGroupID;
+		PolygonGroupToPolygons.AddReferenceToKey(PolygonGroupID, PolygonID);
+
+		for (const FTriangleID TriangleID : PolygonToTriangles.Find(PolygonID))
+		{
+			check(TrianglePolygonGroups[TriangleID] == OldPolygonGroupID);
+			PolygonGroupToTriangles.RemoveReferenceFromKey(OldPolygonGroupID, TriangleID);
+
+			TrianglePolygonGroups[TriangleID] = PolygonGroupID;
+			PolygonGroupToTriangles.AddReferenceToKey(PolygonGroupID, TriangleID);
+		}
 	}
 
 	/** Reverse the winding order of the vertices of this polygon */
 	void ReversePolygonFacing(const FPolygonID PolygonID);
 
+	/** Determines the vertex instances which form the perimeter of a polygon */
+	void FindPolygonPerimeter(TArrayView<const FTriangleID> Triangles, TArrayView<TTuple<int32, int32>> Result) const;
+	void FindPolygonPerimeter(const FPolygonID PolygonID, TArrayView<FEdgeID> Edges) const;
+
 	/** Generates triangles and internal edges for the given polygon */
 	void ComputePolygonTriangulation(const FPolygonID PolygonID);
 
@@ -1197,15 +1113,27 @@ public:
 	// Polygon group operations
 
 	/** Returns the polygons associated with the given polygon group */
-	const TArray<FPolygonID>& GetPolygonGroupPolygons(const FPolygonGroupID PolygonGroupID) const
+	TArrayView<const FPolygonID> GetPolygonGroupPolygons(const FPolygonGroupID PolygonGroupID) const
+	{
+		return PolygonGroupToPolygons.Find<FPolygonID>(PolygonGroupID);
+	}
+
+	/** Returns the triangles associated with the given polygon group */
+	TArrayView<const FTriangleID> GetPolygonGroupTriangles(const FPolygonGroupID PolygonGroupID) const
 	{
-		return PolygonGroupArray[PolygonGroupID].Polygons;
+		return PolygonGroupToTriangles.Find<FTriangleID>(PolygonGroupID);
 	}
 
 	/** Returns the number of polygons in this polygon group */
 	int32 GetNumPolygonGroupPolygons(const FPolygonGroupID PolygonGroupID) const
 	{
-		return PolygonGroupArray[PolygonGroupID].Polygons.Num();
+		return PolygonGroupToPolygons.Find(PolygonGroupID).Num();
+	}
+
+	/** Returns the number of triangles in this polygon group */
+	int32 GetNumPolygonGroupTriangles(const FPolygonGroupID PolygonGroupID) const
+	{
+		return PolygonGroupToTriangles.Find(PolygonGroupID).Num();
 	}
 
 	/** Remaps polygon groups according to the supplied map */
@@ -1221,6 +1149,9 @@ public:
 	/** Remaps the element IDs in the mesh description according to the passed in object */
 	void Remap( const FElementIDRemappings& Remappings );
 
+	/** Sets the specified number of UV channels */
+	void SetNumUVChannels(const int32 NumUVChannels);
+
 	/** Returns bounds of vertices */
 	FBoxSphereBounds GetBounds() const;
 
@@ -1230,23 +1161,28 @@ public:
 	/** Reverses the winding order of all polygons in the mesh */
 	void ReverseAllPolygonFacing();
 
+	float GetTriangleCornerAngleForVertex(const FTriangleID TriangleID, const FVertexID VertexID) const;
 	float GetPolygonCornerAngleForVertex(const FPolygonID PolygonID, const FVertexID VertexID) const;
 
 	FBox ComputeBoundingBox() const;
 
 private:
 
-	FPlane ComputePolygonPlane(const FPolygonID PolygonID) const;
-	FVector ComputePolygonNormal(const FPolygonID PolygonID) const;
+	FPlane ComputePolygonPlane(TArrayView<const FVertexInstanceID> PerimeterVertexInstanceIDs) const;
+	FVector ComputePolygonNormal(TArrayView<const FVertexInstanceID> PerimeterVertexInstanceIDs) const;
 
 private:
 
-	void CreateVertex_Internal(const FVertexID VertexID) { VertexAttributesSet.Insert(VertexID); }
+	void Initialize();
+	void InitializeIndexers();
+	void Cache();
+
 	void CreateVertexInstance_Internal(const FVertexInstanceID VertexInstanceID, const FVertexID VertexID);
 	void CreateEdge_Internal(const FEdgeID EdgeID, const FVertexID VertexID0, const FVertexID VertexID1);
 	void CreateTriangle_Internal(const FTriangleID TriangleID, const FPolygonGroupID PolygonGroupID, TArrayView<const FVertexInstanceID> VertexInstanceIDs, TArray<FEdgeID>* OutEdgeIDs);
 	void CreatePolygon_Internal(const FPolygonID PolygonID, const FPolygonGroupID PolygonGroupID, TArrayView<const FVertexInstanceID> VertexInstanceIDs, TArray<FEdgeID>* OutEdgeIDs);
-	void CreatePolygonGroup_Internal(const FPolygonGroupID PolygonGroupID) { PolygonGroupAttributesSet.Insert(PolygonGroupID); }
+	void CreatePolygonTriangles(const FPolygonID PolygonID, TArrayView<const FVertexInstanceID> VertexInstanceIDs);
+	void RemovePolygonTriangles(const FPolygonID PolygonID);
 
 	template <template <typename...> class TContainer>
 	void DeleteVertexInstance_Internal(const FVertexInstanceID VertexInstanceID, TContainer<FVertexID>* InOutOrphanedVerticesPtr = nullptr);
@@ -1260,9 +1196,6 @@ private:
 	/** Given a set of index remappings, fixes up references to element IDs */
 	void FixUpElementIDs(const FElementIDRemappings& Remappings);
 
-	/** Given a set of index remappings, remaps all attributes accordingly */
-	void RemapAttributes(const FElementIDRemappings& Remappings);
-
 	template <class T>
 	void AddUnique(TArray<T>& Container, const T& Item)
 	{
@@ -1275,20 +1208,44 @@ private:
 		Container.Add(Item);
 	}
 
-
-	FVertexArray VertexArray;
-	FVertexInstanceArray VertexInstanceArray;
-	FEdgeArray EdgeArray;
-	FTriangleArray TriangleArray;
-	FPolygonArray PolygonArray;
-	FPolygonGroupArray PolygonGroupArray;
-
-	TAttributesSet<FVertexID> VertexAttributesSet;
-	TAttributesSet<FVertexInstanceID> VertexInstanceAttributesSet;
-	TAttributesSet<FEdgeID> EdgeAttributesSet;
-	TAttributesSet<FTriangleID> TriangleAttributesSet;
-	TAttributesSet<FPolygonID> PolygonAttributesSet;
-	TAttributesSet<FPolygonGroupID> PolygonGroupAttributesSet;
+	static FName VerticesName;
+	static FName VertexInstancesName;
+	static FName UVsName;
+	static FName EdgesName;
+	static FName TrianglesName;
+	static FName PolygonsName;
+	static FName PolygonGroupsName;
+
+	TMap<FName, FMeshElementTypeWrapper> Elements;
+
+	FMeshElementChannels* VertexElements;
+	FMeshElementChannels* VertexInstanceElements;
+	FMeshElementChannels* UVElements;
+	FMeshElementChannels* EdgeElements;
+	FMeshElementChannels* TriangleElements;
+	FMeshElementChannels* PolygonElements;
+	FMeshElementChannels* PolygonGroupElements;
+
+	TMeshAttributesArray<FVertexID> VertexInstanceVertices;
+	TMeshAttributesArray<TArrayView<FVertexID>> EdgeVertices;
+	TMeshAttributesArray<TArrayView<FVertexInstanceID>> TriangleVertexInstances;
+	TMeshAttributesArray<FPolygonID> TrianglePolygons;
+	TMeshAttributesArray<TArrayView<FEdgeID>> TriangleEdges;
+	TMeshAttributesArray<TArrayView<FVertexID>> TriangleVertices;
+	TMeshAttributesArray<TArrayView<FUVID>> TriangleUVs;
+	TMeshAttributesArray<FPolygonGroupID> TrianglePolygonGroups;
+	TMeshAttributesArray<FPolygonGroupID> PolygonPolygonGroups;
+	TMeshAttributesArray<FVector> VertexPositions;
+
+	// The indexers are all mutable because they may be rebuilt during what is essentially a get operation
+	mutable FMeshElementIndexer VertexToVertexInstances;
+	mutable FMeshElementIndexer VertexToEdges;
+	mutable FMeshElementIndexer VertexInstanceToTriangles;
+	mutable FMeshElementIndexer EdgeToTriangles;
+	mutable FMeshElementIndexer UVToTriangles;
+	mutable FMeshElementIndexer PolygonToTriangles;
+	mutable FMeshElementIndexer PolygonGroupToTriangles;
+	mutable FMeshElementIndexer PolygonGroupToPolygons;
 };
 
 
@@ -1345,13 +1302,3 @@ private:
 	bool bGuidIsHash;
 };
 
-
-UCLASS(deprecated)
-class MESHDESCRIPTION_API UDEPRECATED_MeshDescription : public UObject
-{
-public:
-	GENERATED_BODY()
-
-	// UObject interface
-	virtual void Serialize(FArchive& Ar) override;
-};
diff --git a/Engine/Source/Runtime/MeshDescription/Public/MeshDescriptionBase.h b/Engine/Source/Runtime/MeshDescription/Public/MeshDescriptionBase.h
--- a/Engine/Source/Runtime/MeshDescription/Public/MeshDescriptionBase.h
+++ b/Engine/Source/Runtime/MeshDescription/Public/MeshDescriptionBase.h
@@ -221,7 +221,7 @@ public:
 	UFUNCTION(BlueprintPure, Category="MeshDescription")
 	bool IsVertexOrphaned(FVertexID VertexID) const;
 
-	/** Returns the edge ID defined by the two given vertex IDs, if there is one; otherwise FEdgeID::Invalid */
+	/** Returns the edge ID defined by the two given vertex IDs, if there is one; otherwise INDEX_NONE */
 	UFUNCTION(BlueprintPure, Category="MeshDescription")
 	FEdgeID GetVertexPairEdge(FVertexID VertexID0, FVertexID VertexID1) const;
 
@@ -277,7 +277,7 @@ public:
 	UFUNCTION(BlueprintPure, Category="MeshDescription")
 	FVertexID GetVertexInstanceVertex(FVertexInstanceID VertexInstanceID) const;
 
-	/** Returns the edge ID defined by the two given vertex instance IDs, if there is one; otherwise FEdgeID::Invalid */
+	/** Returns the edge ID defined by the two given vertex instance IDs, if there is one; otherwise INDEX_NONE */
 	UFUNCTION(BlueprintPure, Category="MeshDescription")
 	FEdgeID GetVertexInstancePairEdge(FVertexInstanceID VertexInstanceID0, FVertexInstanceID VertexInstanceID1) const;
 
@@ -369,7 +369,7 @@ public:
 	UFUNCTION(BlueprintPure, Category="MeshDescription")
 	void GetTriangleAdjacentTriangles(FTriangleID TriangleID, TArray<FTriangleID>& OutTriangleIDs) const;
 
-	/** Return the vertex instance which corresponds to the given vertex on the given triangle, or FVertexInstanceID::Invalid */
+	/** Return the vertex instance which corresponds to the given vertex on the given triangle, or INDEX_NONE */
 	UFUNCTION(BlueprintPure, Category="MeshDescription")
 	FVertexInstanceID GetVertexInstanceForTriangleVertex(FTriangleID TriangleID, FVertexID VertexID) const;
 
@@ -418,13 +418,13 @@ public:
 	UFUNCTION(BlueprintPure, Category="MeshDescription")
 	FPolygonGroupID GetPolygonPolygonGroup(FPolygonID PolygonID) const;
 
-	/** Return the vertex instance which corresponds to the given vertex on the given polygon, or FVertexInstanceID::Invalid */
+	/** Return the vertex instance which corresponds to the given vertex on the given polygon, or INDEX_NONE */
 	UFUNCTION(BlueprintPure, Category="MeshDescription")
 	FVertexInstanceID GetVertexInstanceForPolygonVertex(FPolygonID PolygonID, FVertexID VertexID) const;
 
 	/** Set the vertex instance at the given index around the polygon to the new value */
 	UFUNCTION(BlueprintCallable, Category="MeshDescription")
-	void SetPolygonVertexInstance(FPolygonID PolygonID, int32 PerimeterIndex, FVertexInstanceID VertexInstanceID);
+	void SetPolygonVertexInstances(FPolygonID PolygonID, const TArray<FVertexInstanceID>& VertexInstanceIDs);
 
 	/** Sets the polygon group associated with a polygon */
 	UFUNCTION(BlueprintCallable, Category="MeshDescription")
diff --git a/Engine/Source/Runtime/MeshDescription/Public/MeshElementArray.h b/Engine/Source/Runtime/MeshDescription/Public/MeshElementArray.h
--- a/Engine/Source/Runtime/MeshDescription/Public/MeshElementArray.h
+++ b/Engine/Source/Runtime/MeshDescription/Public/MeshElementArray.h
@@ -163,7 +163,7 @@ public:
 	{
 		return Container.Num() > 0 ?
 			ElementIDType( typename TSparseArray<ElementType>::TConstIterator( Container ).GetIndex() ) :
-			ElementIDType::Invalid;
+			INDEX_NONE;
 	}
 
 	/** Returns whether the given ID is valid or not */
@@ -213,7 +213,7 @@ public:
 
 			FORCEINLINE ElementIDType operator*() const
 			{
-				return Iterator ? ElementIDType( Iterator.GetIndex() ) : ElementIDType::Invalid;
+				return Iterator ? ElementIDType( Iterator.GetIndex() ) : INDEX_NONE;
 			}
 
 			friend FORCEINLINE bool operator==( const TConstIterator& Lhs, const TConstIterator& Rhs )
diff --git a/Engine/Source/Runtime/MeshDescription/Public/MeshElementRemappings.h b/Engine/Source/Runtime/MeshDescription/Public/MeshElementRemappings.h
--- a/Engine/Source/Runtime/MeshDescription/Public/MeshElementRemappings.h
+++ b/Engine/Source/Runtime/MeshDescription/Public/MeshElementRemappings.h
@@ -14,6 +14,7 @@ struct FElementIDRemappings
 	TSparseArray<int32> NewVertexInstanceIndexLookup;
 	TSparseArray<int32> NewEdgeIndexLookup;
 	TSparseArray<int32> NewTriangleIndexLookup;
+	TSparseArray<int32> NewUVIndexLookup;
 	TSparseArray<int32> NewPolygonIndexLookup;
 	TSparseArray<int32> NewPolygonGroupIndexLookup;
 
@@ -35,6 +36,12 @@ struct FElementIDRemappings
 		return FEdgeID( NewEdgeIndexLookup[ EdgeID.GetValue() ] );
 	}
 
+	FUVID GetRemappedUVID( FUVID UVID ) const
+	{
+		checkSlow( NewUVIndexLookup.IsAllocated( UVID.GetValue() ) );
+		return FUVID( NewUVIndexLookup[ UVID.GetValue() ] );
+	}
+
 	FTriangleID GetRemappedTriangleID( FTriangleID TriangleID ) const
 	{
 		checkSlow( NewTriangleIndexLookup.IsAllocated( TriangleID.GetValue() ) );
diff --git a/Engine/Source/Runtime/MeshDescription/Public/MeshTypes.h b/Engine/Source/Runtime/MeshDescription/Public/MeshTypes.h
--- a/Engine/Source/Runtime/MeshDescription/Public/MeshTypes.h
+++ b/Engine/Source/Runtime/MeshDescription/Public/MeshTypes.h
@@ -16,11 +16,11 @@ struct FElementID	// @todo mesheditor script: BP doesn't have name spaces, so we
 	GENERATED_BODY()
 
 	FElementID()
-		: IDValue(Invalid.GetValue())
+		: IDValue(INDEX_NONE)
 	{
 	}
 
-	explicit FElementID( const int32 InitIDValue )
+	FElementID( const int32 InitIDValue )
 		: IDValue( InitIDValue )
 	{
 	}
@@ -30,19 +30,34 @@ struct FElementID	// @todo mesheditor script: BP doesn't have name spaces, so we
 		return IDValue;
 	}
 
+	FORCEINLINE operator int32() const
+	{
+		return IDValue;
+	}
+
 	FORCEINLINE bool operator==( const FElementID& Other ) const
 	{
 		return IDValue == Other.IDValue;
 	}
 
+	FORCEINLINE bool operator==( const int32 Other ) const
+	{
+		return IDValue == Other;
+	}
+
 	FORCEINLINE bool operator!=( const FElementID& Other ) const
 	{
 		return IDValue != Other.IDValue;
 	}
 
+	FORCEINLINE bool operator!=( const int32 Other ) const
+	{
+		return IDValue != Other;
+	}
+
 	FString ToString() const
 	{
-		return ( IDValue == Invalid.GetValue() ) ? TEXT( "Invalid" ) : FString::Printf( TEXT( "%d" ), IDValue );
+		return ( IDValue == INDEX_NONE ) ? TEXT( "Invalid" ) : FString::Printf( TEXT( "%d" ), IDValue );
 	}
 
 	friend FArchive& operator<<( FArchive& Ar, FElementID& Element )
@@ -52,6 +67,7 @@ struct FElementID	// @todo mesheditor script: BP doesn't have name spaces, so we
 	}
 
 	/** Invalid element ID */
+	UE_DEPRECATED(4.26, "Please use INDEX_NONE as an invalid ID.")
 	MESHDESCRIPTION_API static const FElementID Invalid;
 
 protected:
@@ -71,12 +87,12 @@ struct FVertexID : public FElementID
 	{
 	}
 
-	explicit FVertexID( const FElementID InitElementID )
+	FVertexID( const FElementID InitElementID )
 		: FElementID( InitElementID.GetValue() )
 	{
 	}
 
-	explicit FVertexID( const int32 InitIDValue )
+	FVertexID( const int32 InitIDValue )
 		: FElementID( InitIDValue )
 	{
 	}
@@ -87,6 +103,7 @@ struct FVertexID : public FElementID
 	}
 
 	/** Invalid vertex ID */
+	UE_DEPRECATED(4.26, "Please use INDEX_NONE as an invalid ID.")
 	MESHDESCRIPTION_API static const FVertexID Invalid;
 };
 
@@ -100,12 +117,12 @@ struct FVertexInstanceID : public FElementID
 	{
 	}
 
-	explicit FVertexInstanceID( const FElementID InitElementID )
+	FVertexInstanceID( const FElementID InitElementID )
 		: FElementID( InitElementID.GetValue() )
 	{
 	}
 
-	explicit FVertexInstanceID( const uint32 InitIDValue )
+	FVertexInstanceID( const int32 InitIDValue )
 		: FElementID( InitIDValue )
 	{
 	}
@@ -116,6 +133,7 @@ struct FVertexInstanceID : public FElementID
 	}
 
 	/** Invalid rendering vertex ID */
+	UE_DEPRECATED(4.26, "Please use INDEX_NONE as an invalid ID.")
 	MESHDESCRIPTION_API static const FVertexInstanceID Invalid;
 };
 
@@ -129,12 +147,12 @@ struct FEdgeID : public FElementID
 	{
 	}
 
-	explicit FEdgeID( const FElementID InitElementID )
+	FEdgeID( const FElementID InitElementID )
 		: FElementID( InitElementID.GetValue() )
 	{
 	}
 
-	explicit FEdgeID( const int32 InitIDValue )
+	FEdgeID( const int32 InitIDValue )
 		: FElementID( InitIDValue )
 	{
 	}
@@ -145,10 +163,37 @@ struct FEdgeID : public FElementID
 	}
 
 	/** Invalid edge ID */
+	UE_DEPRECATED(4.26, "Please use INDEX_NONE as an invalid ID.")
 	MESHDESCRIPTION_API static const FEdgeID Invalid;
 };
 
 
+USTRUCT( BlueprintType )
+struct FUVID : public FElementID
+{
+	GENERATED_BODY()
+
+	FUVID()
+	{
+	}
+
+	FUVID( const FElementID InitElementID )
+		: FElementID( InitElementID.GetValue() )
+	{
+	}
+
+	FUVID( const int32 InitIDValue )
+		: FElementID( InitIDValue )
+	{
+	}
+
+	FORCEINLINE friend uint32 GetTypeHash( const FUVID& Other )
+	{
+		return GetTypeHash( Other.IDValue );
+	}
+};
+
+
 USTRUCT(BlueprintType)
 struct FTriangleID : public FElementID
 {
@@ -158,13 +203,13 @@ struct FTriangleID : public FElementID
 	{
 	}
 
-	explicit FTriangleID(const FElementID InitElementID)
+	FTriangleID(const FElementID InitElementID)
 		: FElementID(InitElementID.GetValue())
 	{
 	}
 
-	explicit FTriangleID(const int32 InitIDValue)
-		: FElementID(InitIDValue)
+	FTriangleID( const int32 InitIDValue )
+		: FElementID( InitIDValue )
 	{
 	}
 
@@ -174,6 +219,7 @@ struct FTriangleID : public FElementID
 	}
 
 	/** Invalid edge ID */
+	UE_DEPRECATED(4.26, "Please use INDEX_NONE as an invalid ID.")
 	MESHDESCRIPTION_API static const FTriangleID Invalid;
 };
 
@@ -187,12 +233,12 @@ struct FPolygonGroupID : public FElementID
 	{
 	}
 
-	explicit FPolygonGroupID( const FElementID InitElementID )
+	FPolygonGroupID( const FElementID InitElementID )
 		: FElementID( InitElementID.GetValue() )
 	{
 	}
 
-	explicit FPolygonGroupID( const int32 InitIDValue )
+	FPolygonGroupID( const int32 InitIDValue )
 		: FElementID( InitIDValue )
 	{
 	}
@@ -203,6 +249,7 @@ struct FPolygonGroupID : public FElementID
 	}
 
 	/** Invalid section ID */
+	UE_DEPRECATED(4.26, "Please use INDEX_NONE as an invalid ID.")
 	MESHDESCRIPTION_API static const FPolygonGroupID Invalid;
 };
 
@@ -216,12 +263,12 @@ struct FPolygonID : public FElementID
 	{
 	}
 
-	explicit FPolygonID( const FElementID InitElementID )
+	FPolygonID( const FElementID InitElementID )
 		: FElementID( InitElementID.GetValue() )
 	{
 	}
 
-	explicit FPolygonID( const int32 InitIDValue )
+	FPolygonID( const int32 InitIDValue )
 		: FElementID( InitIDValue )
 	{
 	}
@@ -232,6 +279,7 @@ struct FPolygonID : public FElementID
 	}
 
 	/** Invalid polygon ID */
+	UE_DEPRECATED(4.26, "Please use INDEX_NONE as an invalid ID.")
 	MESHDESCRIPTION_API static const FPolygonID Invalid;	// @todo mesheditor script: Can we expose these to BP nicely?	Do we even need to?
 };
 
diff --git a/Engine/Source/Runtime/StaticMeshDescription/Private/StaticMeshAttributes.cpp b/Engine/Source/Runtime/StaticMeshDescription/Private/StaticMeshAttributes.cpp
--- a/Engine/Source/Runtime/StaticMeshDescription/Private/StaticMeshAttributes.cpp
+++ b/Engine/Source/Runtime/StaticMeshDescription/Private/StaticMeshAttributes.cpp
@@ -11,11 +11,16 @@ namespace MeshAttribute
 	const FName VertexInstance::Tangent("Tangent");
 	const FName VertexInstance::BinormalSign("BinormalSign");
 	const FName VertexInstance::Color("Color");
+	const FName VertexInstance::UVIndex("UVIndex");
 
 	const FName Edge::IsHard("IsHard");
 	const FName Edge::IsUVSeam("IsUVSeam");
 	const FName Edge::CreaseSharpness("CreaseSharpness");
 
+	const FName Triangle::Normal("Normal");
+	const FName Triangle::Tangent("Tangent");
+	const FName Triangle::Binormal("Binormal");
+
 	const FName Polygon::Normal("Normal");
 	const FName Polygon::Tangent("Tangent");
 	const FName Polygon::Binormal("Binormal");
@@ -31,7 +36,6 @@ namespace MeshAttribute
 void FStaticMeshAttributes::Register()
 {
 	// Add basic vertex attributes
-	MeshDescription.VertexAttributes().RegisterAttribute<float>(MeshAttribute::Vertex::CornerSharpness, 1, 0.0f, EMeshAttributeFlags::Lerpable);
 
 	// Add basic vertex instance attributes
 	MeshDescription.VertexInstanceAttributes().RegisterAttribute<FVector2D>(MeshAttribute::VertexInstance::TextureCoordinate, 1, FVector2D::ZeroVector, EMeshAttributeFlags::Lerpable);
@@ -42,7 +46,6 @@ void FStaticMeshAttributes::Register()
 
 	// Add basic edge attributes
 	MeshDescription.EdgeAttributes().RegisterAttribute<bool>(MeshAttribute::Edge::IsHard, 1, false);
-	MeshDescription.EdgeAttributes().RegisterAttribute<float>(MeshAttribute::Edge::CreaseSharpness, 1, 0.0f, EMeshAttributeFlags::Lerpable);
 
 	// Add basic polygon attributes
 
@@ -56,8 +59,18 @@ void FStaticMeshAttributes::Register()
 
 void FStaticMeshAttributes::RegisterPolygonNormalAndTangentAttributes()
 {
+	PRAGMA_DISABLE_DEPRECATION_WARNINGS
 	MeshDescription.PolygonAttributes().RegisterAttribute<FVector>(MeshAttribute::Polygon::Normal, 1, FVector::ZeroVector, EMeshAttributeFlags::Transient);
 	MeshDescription.PolygonAttributes().RegisterAttribute<FVector>(MeshAttribute::Polygon::Tangent, 1, FVector::ZeroVector, EMeshAttributeFlags::Transient);
 	MeshDescription.PolygonAttributes().RegisterAttribute<FVector>(MeshAttribute::Polygon::Binormal, 1, FVector::ZeroVector, EMeshAttributeFlags::Transient);
 	MeshDescription.PolygonAttributes().RegisterAttribute<FVector>(MeshAttribute::Polygon::Center, 1, FVector::ZeroVector, EMeshAttributeFlags::Transient);
+	PRAGMA_ENABLE_DEPRECATION_WARNINGS
+}
+
+
+void FStaticMeshAttributes::RegisterTriangleNormalAndTangentAttributes()
+{
+	MeshDescription.TriangleAttributes().RegisterAttribute<FVector>(MeshAttribute::Triangle::Normal, 1, FVector::ZeroVector, EMeshAttributeFlags::Transient);
+	MeshDescription.TriangleAttributes().RegisterAttribute<FVector>(MeshAttribute::Triangle::Tangent, 1, FVector::ZeroVector, EMeshAttributeFlags::Transient);
+	MeshDescription.TriangleAttributes().RegisterAttribute<FVector>(MeshAttribute::Triangle::Binormal, 1, FVector::ZeroVector, EMeshAttributeFlags::Transient);
 }
diff --git a/Engine/Source/Runtime/StaticMeshDescription/Private/StaticMeshDescription.cpp b/Engine/Source/Runtime/StaticMeshDescription/Private/StaticMeshDescription.cpp
--- a/Engine/Source/Runtime/StaticMeshDescription/Private/StaticMeshDescription.cpp
+++ b/Engine/Source/Runtime/StaticMeshDescription/Private/StaticMeshDescription.cpp
@@ -127,12 +127,7 @@ void UStaticMeshDescription::CreateCube(FVector Center, FVector HalfExtents, FPo
 	PolygonID_PlusZ = MakePolygon(1, 0, 3, 2);
 	PolygonID_MinusZ = MakePolygon(6, 7, 4, 5);
 
-	MeshDescription.PolygonAttributes().RegisterAttribute<FVector>(MeshAttribute::Polygon::Normal, 1, FVector::ZeroVector, EMeshAttributeFlags::Transient);
-	MeshDescription.PolygonAttributes().RegisterAttribute<FVector>(MeshAttribute::Polygon::Tangent, 1, FVector::ZeroVector, EMeshAttributeFlags::Transient);
-	MeshDescription.PolygonAttributes().RegisterAttribute<FVector>(MeshAttribute::Polygon::Binormal, 1, FVector::ZeroVector, EMeshAttributeFlags::Transient);
-	MeshDescription.PolygonAttributes().RegisterAttribute<FVector>(MeshAttribute::Polygon::Center, 1, FVector::ZeroVector, EMeshAttributeFlags::Transient);
-
-	FStaticMeshOperations::ComputePolygonTangentsAndNormals(MeshDescription);
+	FStaticMeshOperations::ComputeTriangleTangentsAndNormals(MeshDescription);
 	FStaticMeshOperations::ComputeTangentsAndNormals(MeshDescription, EComputeNTBsFlags::Normals | EComputeNTBsFlags::Tangents);
 }
 
