Unbroke static mesh serialization.

diff --git a/Engine/Source/Runtime/MeshDescription/Private/MeshAttributeArray.cpp b/Engine/Source/Runtime/MeshDescription/Private/MeshAttributeArray.cpp
--- a/Engine/Source/Runtime/MeshDescription/Private/MeshAttributeArray.cpp
+++ b/Engine/Source/Runtime/MeshDescription/Private/MeshAttributeArray.cpp
@@ -13,7 +13,7 @@ FArchive& operator<<(FArchive& Ar, FAttributesSetEntry& Entry)
 		Ar << AttributeType;
 
 		uint32 Extent = 1;
-		if (Ar.CustomVer(FEditorObjectVersion::GUID) >= FReleaseObjectVersion::MeshDescriptionNewFormat)
+		if (Ar.CustomVer(FReleaseObjectVersion::GUID) >= FReleaseObjectVersion::MeshDescriptionNewFormat)
 		{
 			Ar << Extent;
 		}
diff --git a/Engine/Source/Runtime/MeshDescription/Private/MeshDescription.cpp b/Engine/Source/Runtime/MeshDescription/Private/MeshDescription.cpp
--- a/Engine/Source/Runtime/MeshDescription/Private/MeshDescription.cpp
+++ b/Engine/Source/Runtime/MeshDescription/Private/MeshDescription.cpp
@@ -162,7 +162,7 @@ void FMeshDescription::Serialize(FArchive& Ar)
 		UE_LOG(LogLoad, Warning, TEXT("Deprecated serialization format"));
 	}
 
-	if (Ar.IsLoading() && Ar.CustomVer(FEditorObjectVersion::GUID) < FReleaseObjectVersion::MeshDescriptionNewFormat)
+	if (Ar.IsLoading() && Ar.CustomVer(FReleaseObjectVersion::GUID) < FReleaseObjectVersion::MeshDescriptionNewFormat)
 	{
 		// Serialize the old format data and transform it into the new format mesh element map
 		SerializeLegacy(Ar);
@@ -1889,6 +1889,7 @@ void FMeshDescription::RemapPolygonGroups(const TMap<FPolygonGroupID, FPolygonGr
 void FMeshDescriptionBulkData::Serialize( FArchive& Ar, UObject* Owner )
 {
 	Ar.UsingCustomVersion( FEditorObjectVersion::GUID );
+	Ar.UsingCustomVersion( FReleaseObjectVersion::GUID );
 	Ar.UsingCustomVersion( FEnterpriseObjectVersion::GUID );
 
 	if( Ar.IsTransacting() )
diff --git a/Engine/Source/Runtime/MeshDescription/Public/MeshAttributeArray.h b/Engine/Source/Runtime/MeshDescription/Public/MeshAttributeArray.h
--- a/Engine/Source/Runtime/MeshDescription/Public/MeshAttributeArray.h
+++ b/Engine/Source/Runtime/MeshDescription/Public/MeshAttributeArray.h
@@ -219,7 +219,7 @@ void TMeshAttributeArrayBase<AttributeType>::Remap(const TSparseArray<int32>& In
 template <typename T>
 inline typename TEnableIf<!TIsBulkSerializable<T>::Value, FArchive>::Type& operator<<( FArchive& Ar, TMeshAttributeArrayBase<T>& Array )
 {
-	if (Ar.IsLoading() && Ar.CustomVer(FEditorObjectVersion::GUID) < FReleaseObjectVersion::MeshDescriptionNewFormat)
+	if (Ar.IsLoading() && Ar.CustomVer(FReleaseObjectVersion::GUID) < FReleaseObjectVersion::MeshDescriptionNewFormat)
 	{
 		Array.Extent = 1;
 	}
@@ -236,7 +236,7 @@ inline typename TEnableIf<!TIsBulkSerializable<T>::Value, FArchive>::Type& opera
 template <typename T>
 inline typename TEnableIf<TIsBulkSerializable<T>::Value, FArchive>::Type& operator<<( FArchive& Ar, TMeshAttributeArrayBase<T>& Array )
 {
-	if (Ar.IsLoading() && Ar.CustomVer(FEditorObjectVersion::GUID) < FReleaseObjectVersion::MeshDescriptionNewFormat)
+	if (Ar.IsLoading() && Ar.CustomVer(FReleaseObjectVersion::GUID) < FReleaseObjectVersion::MeshDescriptionNewFormat)
 	{
 		Array.Extent = 1;
 	}
