Fixed issue when serializing legacy assets.

diff --git a/Engine/Source/Runtime/MeshDescription/Private/MeshDescription.cpp b/Engine/Source/Runtime/MeshDescription/Private/MeshDescription.cpp
--- a/Engine/Source/Runtime/MeshDescription/Private/MeshDescription.cpp
+++ b/Engine/Source/Runtime/MeshDescription/Private/MeshDescription.cpp
@@ -328,9 +328,38 @@ struct FMeshPolygonGroup_Legacy
 };
 
 
+template <typename T>
+void FixAttributesSize(int32 ExpectedNum, TAttributesSet<T>& AttributesSet)
+{
+	// Ensure that the attribute set is the same size as the mesh element array they describe
+	// If there are extra elements, and they are not set to trivial defaults, this is an error.
+	bool bAllDefault = true;
+	AttributesSet.ForEach([ExpectedNum, &bAllDefault](const FName AttributeName, auto AttributesConstRef)
+		{
+			if (bAllDefault)
+			{
+				for (int32 Channel = 0; Channel < AttributesConstRef.GetNumChannels(); Channel++)
+				{
+					for (int32 Index = ExpectedNum; Index < AttributesConstRef.GetNumElements(); Index++)
+					{
+						if (AttributesConstRef.Get(T(Index), Channel) != AttributesConstRef.GetDefaultValue())
+						{
+							bAllDefault = false;
+							break;
+						}
+					}
+				}
+			}
+		}
+	);
+
+	check(bAllDefault);	// If this fires, something is very wrong with the legacy asset
+	AttributesSet.SetNumElements(ExpectedNum);
+}
+
 void FMeshDescription::SerializeLegacy(FArchive& Ar)
 {
-	TMeshElementArray<FMeshVertex_Legacy, FVertexID> VertexArray;
+	TMeshElementArray<FMeshVertex_Legacy, FVertexID> VertexArray; 
 	TMeshElementArray<FMeshVertexInstance_Legacy, FVertexInstanceID> VertexInstanceArray;
 	TMeshElementArray<FMeshEdge_Legacy, FEdgeID> EdgeArray;
 	TMeshElementArray<FMeshTriangle_Legacy, FTriangleID> TriangleArray;
@@ -356,11 +385,19 @@ void FMeshDescription::SerializeLegacy(FArchive& Ar)
 	Ar << PolygonAttributesSet;
 	Ar << PolygonGroupAttributesSet;
 
+	FixAttributesSize(VertexArray.GetArraySize(), VertexAttributesSet);
+	FixAttributesSize(VertexInstanceArray.GetArraySize(), VertexInstanceAttributesSet);
+	FixAttributesSize(EdgeArray.GetArraySize(), EdgeAttributesSet);
+	FixAttributesSize(PolygonArray.GetArraySize(), PolygonAttributesSet);
+	FixAttributesSize(PolygonGroupArray.GetArraySize(), PolygonGroupAttributesSet);
+
 	// Serialize new triangle arrays since version MeshDescriptionTriangles
 	if (Ar.CustomVer(FEditorObjectVersion::GUID) >= FEditorObjectVersion::MeshDescriptionTriangles)
 	{
 		Ar << TriangleArray;
 		Ar << TriangleAttributesSet;
+
+		FixAttributesSize(TriangleArray.GetArraySize(), TriangleAttributesSet);
 	}
 
 	// Convert the old style element arrays into the new format
diff --git a/Engine/Source/Runtime/MeshDescription/Public/MeshAttributeArray.h b/Engine/Source/Runtime/MeshDescription/Public/MeshAttributeArray.h
--- a/Engine/Source/Runtime/MeshDescription/Public/MeshAttributeArray.h
+++ b/Engine/Source/Runtime/MeshDescription/Public/MeshAttributeArray.h
@@ -132,6 +132,18 @@ public:
 		}
 	}
 
+	void SetNum(const int32 ElementCount, const AttributeType& Default)
+	{
+		if (Container.Num() >= ElementCount)
+		{
+			Container.SetNum(ElementCount);
+		}
+		else
+		{
+			Initialize(ElementCount, Default);
+		}
+	}
+
 	uint32 GetHash(uint32 Crc = 0) const
 	{
 		return FCrc::MemCrc32(Container.GetData(), Container.Num() * sizeof(AttributeType), Crc);
@@ -306,6 +318,7 @@ public:
 	virtual void Insert(const int32 Index) = 0;
 	virtual void Remove(const int32 Index) = 0;
 	virtual void Initialize(const int32 Count) = 0;
+	virtual void SetNumElements(const int32 Count) = 0;
 	virtual uint32 GetHash() const = 0;
 	virtual void Serialize(FArchive& Ar) = 0;
 	virtual void Remap(const TSparseArray<int32>& IndexRemap) = 0;
@@ -418,6 +431,16 @@ public:
 		}
 	}
 
+	/** Sets the number of elements to the exact number provided, preserving existing elements if the number is bigger */
+	virtual void SetNumElements(const int32 Count) override
+	{
+		NumElements = Count;
+		for (TMeshAttributeArrayBase<AttributeType>& ArrayForChannel : ArrayForChannels)
+		{
+			ArrayForChannel.SetNum(Count, DefaultValue);
+		}
+	}
+
 	virtual uint32 GetHash() const override
 	{
 		uint32 CrcResult = 0;
@@ -1198,6 +1221,22 @@ public:
 		}
 	}
 
+	/** Sets all attributes to have the given number of elements, preserving existing values and filling extra elements with the default value */
+	void SetNumElements(const int32 Count)
+	{
+		NumElements = Count;
+		for (auto& MapEntry : Map)
+		{
+			MapEntry.Value->SetNumElements(Count);
+		}
+	}
+
+	/** Gets the number of elements in the attribute set */
+	int32 GetNumElements() const
+	{
+		return NumElements;
+	}
+
 	/** Applies the given remapping to the attributes set */
 	void Remap(const TSparseArray<int32>& IndexRemap);
 
