Created new object version in UE5MainStreamObjectVersion for MeshDescription format changes.
This is the correct versioning for new MeshDescription objects.
The original versioning (in ReleaseObjectVersion) clashes with changes being brought over from 4.26.  MeshDescription objects already saved in UE5 will only be valid for precisely FReleaseObjectVersion::MeshDescriptionNewFormat; subsequent version numbers belong to assets from 4.26 (and derived streams) which have been resaved and merged into UE5, and do *not* have the new MeshDescription format.
Going forward, UE5MainStreamObjectVersion holds the latest version for MeshDescription objects.

diff --git a/Engine/Source/Runtime/Core/Public/UObject/UE5MainStreamObjectVersion.h b/Engine/Source/Runtime/Core/Public/UObject/UE5MainStreamObjectVersion.h
--- a/Engine/Source/Runtime/Core/Public/UObject/UE5MainStreamObjectVersion.h
+++ b/Engine/Source/Runtime/Core/Public/UObject/UE5MainStreamObjectVersion.h
@@ -21,6 +21,10 @@ struct CORE_API FUE5MainStreamObjectVersion
 		// Removing SourceAnimationData, animation layering is now applied during compression
 		RemovingSourceAnimationData,
 
+		// New MeshDescription format.
+		// This is the correct versioning for MeshDescription changes which were added to ReleaseObjectVersion.
+		MeshDescriptionNewFormat,
+
 		// -----<new versions can be added above this line>-------------------------------------------------
 		VersionPlusOne,
 		LatestVersion = VersionPlusOne - 1
diff --git a/Engine/Source/Runtime/MeshDescription/Private/MeshAttributeArray.cpp b/Engine/Source/Runtime/MeshDescription/Private/MeshAttributeArray.cpp
--- a/Engine/Source/Runtime/MeshDescription/Private/MeshAttributeArray.cpp
+++ b/Engine/Source/Runtime/MeshDescription/Private/MeshAttributeArray.cpp
@@ -3,6 +3,7 @@
 #include "MeshAttributeArray.h"
 #include "UObject/EditorObjectVersion.h"
 #include "UObject/ReleaseObjectVersion.h"
+#include "UObject/UE5MainStreamObjectVersion.h"
 
 
 FArchive& operator<<(FArchive& Ar, FAttributesSetEntry& Entry)
@@ -13,7 +14,8 @@ FArchive& operator<<(FArchive& Ar, FAttributesSetEntry& Entry)
 		Ar << AttributeType;
 
 		uint32 Extent = 1;
-		if (Ar.CustomVer(FReleaseObjectVersion::GUID) >= FReleaseObjectVersion::MeshDescriptionNewFormat)
+		if (Ar.CustomVer(FReleaseObjectVersion::GUID) == FReleaseObjectVersion::MeshDescriptionNewFormat ||
+			Ar.CustomVer(FUE5MainStreamObjectVersion::GUID) >= FUE5MainStreamObjectVersion::MeshDescriptionNewFormat)
 		{
 			Ar << Extent;
 		}
diff --git a/Engine/Source/Runtime/MeshDescription/Private/MeshDescription.cpp b/Engine/Source/Runtime/MeshDescription/Private/MeshDescription.cpp
--- a/Engine/Source/Runtime/MeshDescription/Private/MeshDescription.cpp
+++ b/Engine/Source/Runtime/MeshDescription/Private/MeshDescription.cpp
@@ -7,6 +7,7 @@
 #include "Serialization/BulkDataReader.h"
 #include "Serialization/BulkDataWriter.h"
 #include "UObject/EnterpriseObjectVersion.h"
+#include "UObject/UE5MainStreamObjectVersion.h"
 
 
 FName FMeshDescription::VerticesName("Vertices");
@@ -156,13 +157,16 @@ void FMeshDescription::Serialize(FArchive& Ar)
 {
 	Ar.UsingCustomVersion(FReleaseObjectVersion::GUID);
 	Ar.UsingCustomVersion(FEditorObjectVersion::GUID);
+	Ar.UsingCustomVersion(FUE5MainStreamObjectVersion::GUID);
 
 	if (Ar.IsLoading() && Ar.CustomVer(FReleaseObjectVersion::GUID) < FReleaseObjectVersion::MeshDescriptionNewSerialization)
 	{
 		UE_LOG(LogLoad, Warning, TEXT("Deprecated serialization format"));
 	}
 
-	if (Ar.IsLoading() && Ar.CustomVer(FReleaseObjectVersion::GUID) < FReleaseObjectVersion::MeshDescriptionNewFormat)
+	if (Ar.IsLoading() &&
+		Ar.CustomVer(FReleaseObjectVersion::GUID) != FReleaseObjectVersion::MeshDescriptionNewFormat &&
+		Ar.CustomVer(FUE5MainStreamObjectVersion::GUID) < FUE5MainStreamObjectVersion::MeshDescriptionNewFormat)
 	{
 		// Serialize the old format data and transform it into the new format mesh element map
 		SerializeLegacy(Ar);
diff --git a/Engine/Source/Runtime/MeshDescription/Public/MeshAttributeArray.h b/Engine/Source/Runtime/MeshDescription/Public/MeshAttributeArray.h
--- a/Engine/Source/Runtime/MeshDescription/Public/MeshAttributeArray.h
+++ b/Engine/Source/Runtime/MeshDescription/Public/MeshAttributeArray.h
@@ -10,6 +10,7 @@
 #include "Templates/CopyQualifiersFromTo.h"
 #include "UObject/EditorObjectVersion.h"
 #include "UObject/ReleaseObjectVersion.h"
+#include "UObject/UE5MainStreamObjectVersion.h"
 #include "Templates/IsArray.h"
 
 
@@ -231,7 +232,9 @@ void TMeshAttributeArrayBase<AttributeType>::Remap(const TSparseArray<int32>& In
 template <typename T>
 inline typename TEnableIf<!TIsBulkSerializable<T>::Value, FArchive>::Type& operator<<( FArchive& Ar, TMeshAttributeArrayBase<T>& Array )
 {
-	if (Ar.IsLoading() && Ar.CustomVer(FReleaseObjectVersion::GUID) < FReleaseObjectVersion::MeshDescriptionNewFormat)
+	if (Ar.IsLoading() &&
+		Ar.CustomVer(FReleaseObjectVersion::GUID) != FReleaseObjectVersion::MeshDescriptionNewFormat &&
+		Ar.CustomVer(FUE5MainStreamObjectVersion::GUID) < FUE5MainStreamObjectVersion::MeshDescriptionNewFormat)
 	{
 		Array.Extent = 1;
 	}
@@ -248,7 +251,9 @@ inline typename TEnableIf<!TIsBulkSerializable<T>::Value, FArchive>::Type& opera
 template <typename T>
 inline typename TEnableIf<TIsBulkSerializable<T>::Value, FArchive>::Type& operator<<( FArchive& Ar, TMeshAttributeArrayBase<T>& Array )
 {
-	if (Ar.IsLoading() && Ar.CustomVer(FReleaseObjectVersion::GUID) < FReleaseObjectVersion::MeshDescriptionNewFormat)
+	if (Ar.IsLoading() &&
+		Ar.CustomVer(FReleaseObjectVersion::GUID) != FReleaseObjectVersion::MeshDescriptionNewFormat &&
+		Ar.CustomVer(FUE5MainStreamObjectVersion::GUID) < FUE5MainStreamObjectVersion::MeshDescriptionNewFormat)
 	{
 		Array.Extent = 1;
 	}
