Added correct custom version for MeshDescriptionBulkData.

diff --git a/Engine/Source/Runtime/MeshDescription/Private/MeshDescription.cpp b/Engine/Source/Runtime/MeshDescription/Private/MeshDescription.cpp
--- a/Engine/Source/Runtime/MeshDescription/Private/MeshDescription.cpp
+++ b/Engine/Source/Runtime/MeshDescription/Private/MeshDescription.cpp
@@ -1924,6 +1924,7 @@ void FMeshDescriptionBulkData::Serialize( FArchive& Ar, UObject* Owner )
 	Ar.UsingCustomVersion( FEditorObjectVersion::GUID );
 	Ar.UsingCustomVersion( FReleaseObjectVersion::GUID );
 	Ar.UsingCustomVersion( FEnterpriseObjectVersion::GUID );
+	Ar.UsingCustomVersion( FUE5MainStreamObjectVersion::GUID );
 
 	if( Ar.IsTransacting() )
 	{
