Copying //UE5/Dev-Cooker@14539516 to Main (//UE5/Main)

diff --git a/Engine/Source/Runtime/MeshDescription/Private/MeshDescription.cpp b/Engine/Source/Runtime/MeshDescription/Private/MeshDescription.cpp
--- a/Engine/Source/Runtime/MeshDescription/Private/MeshDescription.cpp
+++ b/Engine/Source/Runtime/MeshDescription/Private/MeshDescription.cpp
@@ -7,7 +7,9 @@
 #include "Serialization/BulkDataReader.h"
 #include "Serialization/BulkDataWriter.h"
 #include "UObject/EnterpriseObjectVersion.h"
+#include "UObject/UE5CookerObjectVersion.h"
 #include "UObject/UE5MainStreamObjectVersion.h"
+#include "Virtualization/VirtualizationUtilities.h"
 
 
 FName FMeshDescription::VerticesName("Vertices");
@@ -1916,7 +1918,6 @@ void FMeshDescription::RemapPolygonGroups(const TMap<FPolygonGroupID, FPolygonGr
 }
 
 
-
 #if WITH_EDITORONLY_DATA
 
 void FMeshDescriptionBulkData::Serialize( FArchive& Ar, UObject* Owner )
@@ -1925,6 +1926,9 @@ void FMeshDescriptionBulkData::Serialize( FArchive& Ar, UObject* Owner )
 	Ar.UsingCustomVersion( FReleaseObjectVersion::GUID );
 	Ar.UsingCustomVersion( FEnterpriseObjectVersion::GUID );
 	Ar.UsingCustomVersion( FUE5MainStreamObjectVersion::GUID );
+#if UE_USE_VIRTUALBULKDATA
+	Ar.UsingCustomVersion(FUE5CookerObjectVersion::GUID);
+#endif //UE_USE_VIRTUALBULKDATA
 
 	if( Ar.IsTransacting() )
 	{
@@ -1953,7 +1957,20 @@ void FMeshDescriptionBulkData::Serialize( FArchive& Ar, UObject* Owner )
 		}
 	}
 
-	BulkData.Serialize( Ar, Owner );
+#if UE_USE_VIRTUALBULKDATA
+	bool bSerializedOldDataTypes = false;
+	FByteBulkData TempBulkData;
+	if (Ar.IsLoading() && Ar.CustomVer(FUE5CookerObjectVersion::GUID) < FUE5CookerObjectVersion::MeshDescriptionVirtualization)
+	{
+		// Serialize the old BulkData format and mark that we require a conversion after the guid has been serialized
+		TempBulkData.Serialize(Ar, Owner);
+		bSerializedOldDataTypes = true;
+	}
+	else
+#endif //UE_USE_VIRTUALBULKDATA	
+	{
+		BulkData.Serialize(Ar, Owner);
+	}
 
 	if( Ar.IsLoading() && Ar.CustomVer( FEditorObjectVersion::GUID ) < FEditorObjectVersion::MeshDescriptionBulkDataGuid )
 	{
@@ -1973,6 +1990,14 @@ void FMeshDescriptionBulkData::Serialize( FArchive& Ar, UObject* Owner )
 	{
 		Ar << bGuidIsHash;
 	}
+
+#if UE_USE_VIRTUALBULKDATA
+	// Needs to be after the guid is serialized
+	if (bSerializedOldDataTypes)
+	{
+		BulkData.CreateFromBulkData(TempBulkData, Guid);
+	}
+#endif //UE_USE_VIRTUALBULKDATA
 }
 
 
@@ -1980,12 +2005,21 @@ void FMeshDescriptionBulkData::SaveMeshDescription( FMeshDescription& MeshDescri
 {
 	TRACE_CPUPROFILER_EVENT_SCOPE(FMeshDescriptionBulkData::SaveMeshDescription);
 
+#if UE_USE_VIRTUALBULKDATA
+	BulkData.Reset();
+#else
 	BulkData.RemoveBulkData();
+#endif //UE_USE_VIRTUALBULKDATA
 
 	if( !MeshDescription.IsEmpty() )
 	{
 		const bool bIsPersistent = true;
-		FBulkDataWriter Ar( BulkData, bIsPersistent );
+#if UE_USE_VIRTUALBULKDATA
+		FVirtualizedBulkDataWriter Ar(BulkData, bIsPersistent);
+#else
+		FBulkDataWriter Ar(BulkData, bIsPersistent);
+#endif //UE_USE_VIRTUALBULKDATA
+
 		Ar << MeshDescription;
 
 		// Preserve CustomVersions at save time so we can reuse the same ones when reloading direct from memory
@@ -2011,12 +2045,16 @@ void FMeshDescriptionBulkData::LoadMeshDescription( FMeshDescription& MeshDescri
 {
 	MeshDescription.Empty();
 
-	if( BulkData.GetElementCount() > 0 )
+	if (BulkData.GetBulkDataSize() > 0)
 	{
 		// Get a lock on the bulk data and read it into the mesh description
 		{
 			const bool bIsPersistent = true;
-			FBulkDataReader Ar( BulkData, bIsPersistent );
+#if UE_USE_VIRTUALBULKDATA
+			FVirtualizedBulkDataReader Ar(BulkData, bIsPersistent);
+#else
+			FBulkDataReader Ar(BulkData, bIsPersistent);
+#endif //UE_USE_VIRTUALBULKDATA
 
 			// Propagate the custom version information from the package to the bulk data, so that the MeshDescription
 			// is serialized with the same versioning.
@@ -2027,17 +2065,21 @@ void FMeshDescriptionBulkData::LoadMeshDescription( FMeshDescription& MeshDescri
 
 		// Throw away the bulk data allocation as we don't need it now we have its contents as a FMeshDescription
 		// @todo: revisit this
-//		BulkData.UnloadBulkData();
+#if UE_USE_VIRTUALBULKDATA
+		BulkData.UnloadData();
+#endif //UE_USE_VIRTUALBULKDATA
 	}
 }
 
-
 void FMeshDescriptionBulkData::Empty()
 {
+#if UE_USE_VIRTUALBULKDATA
+	BulkData.Reset();
+#else
 	BulkData.RemoveBulkData();
+#endif //UE_USE_VIRTUALBULKDATA
 }
 
-
 FString FMeshDescriptionBulkData::GetIdString() const
 {
 	FString GuidString = Guid.ToString();
@@ -2051,17 +2093,26 @@ FString FMeshDescriptionBulkData::GetIdString() const
 
 void FMeshDescriptionBulkData::UseHashAsGuid()
 {
-	uint32 Hash[5] = {};
-
 	if (BulkData.GetBulkDataSize() > 0)
 	{
 		bGuidIsHash = true;
-		void* Buffer = BulkData.Lock(LOCK_READ_ONLY);
+
+#if UE_USE_VIRTUALBULKDATA
+		Guid = BulkData.GetKey();
+#else
+		uint32 Hash[5] = {};
+
+		const void* Buffer = BulkData.LockReadOnly();
 		FSHA1::HashBuffer(Buffer, BulkData.GetBulkDataSize(), (uint8*)Hash);
 		BulkData.Unlock();
-	}
 
-	Guid = FGuid(Hash[0] ^ Hash[4], Hash[1], Hash[2], Hash[3]);
+		Guid = FGuid(Hash[0] ^ Hash[4], Hash[1], Hash[2], Hash[3]);
+#endif //UE_USE_VIRTUALBULKDATA
+	}
+	else
+	{
+		Guid.Invalidate();
+	}	
 }
 
 #endif // #if WITH_EDITORONLY_DATA
diff --git a/Engine/Source/Runtime/MeshDescription/Public/MeshDescription.h b/Engine/Source/Runtime/MeshDescription/Public/MeshDescription.h
--- a/Engine/Source/Runtime/MeshDescription/Public/MeshDescription.h
+++ b/Engine/Source/Runtime/MeshDescription/Public/MeshDescription.h
@@ -19,6 +19,7 @@
 #include "Serialization/BulkData.h"
 #include "Serialization/CustomVersion.h"
 #include "Containers/StaticArray.h"
+#include "Virtualization/VirtualizedBulkData.h"
 #include "MeshDescription.generated.h"
 
 enum
@@ -1289,6 +1290,7 @@ private:
 	mutable FMeshElementIndexer PolygonGroupToPolygons;
 };
 
+#if WITH_EDITORONLY_DATA
 
 /**
  * Bulk data storage for FMeshDescription
@@ -1300,10 +1302,13 @@ public:
 		: bBulkDataUpdated(false)
 		, bGuidIsHash(false)
 	{
+#if UE_USE_VIRTUALBULKDATA
+		BulkData.EnableDataCompression(true);
+#else
 		BulkData.SetBulkDataFlags(BULKDATA_SerializeCompressed | BULKDATA_SerializeCompressedBitWindow);
+#endif //UE_USE_VIRTUALBULKDATA
 	}
 
-#if WITH_EDITORONLY_DATA
 	/** Serialization */
 	void Serialize(FArchive& Ar, UObject* Owner);
 
@@ -1324,11 +1329,10 @@ public:
 
 	/** Uses a hash as the GUID, useful to prevent recomputing content already in cache. */
 	void UseHashAsGuid();
-#endif
 
 private:
 	/** Internally store bulk data as bytes */
-	FByteBulkData BulkData;
+	TChooseClass<UE_USE_VIRTUALBULKDATA, FByteVirtualizedBulkData, FByteBulkData>::Result BulkData;
 
 	/** GUID associated with the data stored herein. */
 	FGuid Guid;
@@ -1343,3 +1347,4 @@ private:
 	bool bGuidIsHash;
 };
 
+#endif //WITH_EDITORONLY_DATA
diff --git a/Engine/Source/Runtime/MeshDescription/Public/MeshElementArray.h b/Engine/Source/Runtime/MeshDescription/Public/MeshElementArray.h
--- a/Engine/Source/Runtime/MeshDescription/Public/MeshElementArray.h
+++ b/Engine/Source/Runtime/MeshDescription/Public/MeshElementArray.h
@@ -47,13 +47,20 @@ public:
 			// We have to build the TBitArray representing allocated indices by hand, as we don't have access to it from outside TSparseArray.
 			// @todo core: consider replacing TSparseArray serialization with this format.
 			TBitArray<> AllocatedIndices( false, MaxIndex );
+			int32 MaxAllocatedIndex = 0;
+			
 			for( int32 Index = 0; Index < MaxIndex; ++Index )
 			{
 				if( Array.Container.IsAllocated( Index ) )
 				{
 					AllocatedIndices[ Index ] = true;
+					MaxAllocatedIndex = Index;
 				}
 			}
+
+			// Cut off the trailing unallocated indices so that they are not serialized
+			AllocatedIndices.SetNumUninitialized(MaxAllocatedIndex + 1);
+
 			Ar << AllocatedIndices;
 
 			for( auto It = Array.Container.CreateIterator(); It; ++It )
diff --git a/Engine/Source/Runtime/MovieScene/Private/EntitySystem/MovieSceneSequenceInstance.cpp b/Engine/Source/Runtime/MovieScene/Private/EntitySystem/MovieSceneSequenceInstance.cpp
--- a/Engine/Source/Runtime/MovieScene/Private/EntitySystem/MovieSceneSequenceInstance.cpp
+++ b/Engine/Source/Runtime/MovieScene/Private/EntitySystem/MovieSceneSequenceInstance.cpp
@@ -285,7 +285,7 @@ FInstanceHandle FSequenceInstance::FindSubInstance(FMovieSceneSequenceID SubSequ
 
 FMovieSceneEntityID FSequenceInstance::FindEntity(UObject* Owner, uint32 EntityID) const
 {
-	return Ledger.FindImportedEntity(FMovieSceneEvaluationFieldEntityKey{ Owner, EntityID });
+	return Ledger.FindImportedEntity(FMovieSceneEvaluationFieldEntityKey{ decltype(FMovieSceneEvaluationFieldEntityKey::EntityOwner)(Owner), EntityID });
 }
 
 } // namespace MovieScene
diff --git a/Engine/Source/Runtime/MovieScene/Private/MovieSceneSection.cpp b/Engine/Source/Runtime/MovieScene/Private/MovieSceneSection.cpp
--- a/Engine/Source/Runtime/MovieScene/Private/MovieSceneSection.cpp
+++ b/Engine/Source/Runtime/MovieScene/Private/MovieSceneSection.cpp
@@ -314,7 +314,7 @@ void UMovieSceneSection::BuildDefaultComponents(UMovieSceneEntitySystemLinker* E
 
 	OutImportedEntity->AddBuilder(
 		FEntityBuilder()
-		.AddConditional(Components->Easing,                     FEasingComponentData{ this }, bHasEasing)
+		.AddConditional(Components->Easing,                     FEasingComponentData{ decltype(FEasingComponentData::Section)(this) }, bHasEasing)
 		.AddConditional(Components->HierarchicalEasingChannel, uint16(-1), Params.Sequence.bHasHierarchicalEasing)
 		.AddConditional(Components->HierarchicalBias,           Params.Sequence.HierarchicalBias, Params.Sequence.HierarchicalBias != 0)
 		.AddConditional(Components->Interrogation.InputKey,     Params.InterrogationKey, Params.InterrogationKey.IsValid())
diff --git a/Engine/Source/Runtime/MovieScene/Private/MovieSceneTrack.cpp b/Engine/Source/Runtime/MovieScene/Private/MovieSceneTrack.cpp
--- a/Engine/Source/Runtime/MovieScene/Private/MovieSceneTrack.cpp
+++ b/Engine/Source/Runtime/MovieScene/Private/MovieSceneTrack.cpp
@@ -385,7 +385,7 @@ void UMovieSceneTrack::UpdateEvaluationTree()
 			for (const FMovieSceneTrackEvaluationData& TrackData : TrackDataIt)
 			{
 				UMovieSceneSection* Section = TrackData.Section.Get();
-				SectionToEntry.FindOrAdd(Section).Add(FMovieSceneTrackEvaluationFieldEntry{ Section, Range, TrackData.ForcedTime, TrackData.Flags, TrackData.SortOrder });
+				SectionToEntry.FindOrAdd(Section).Add(FMovieSceneTrackEvaluationFieldEntry{decltype(FMovieSceneTrackEvaluationFieldEntry::Section)(Section), Range, TrackData.ForcedTime, TrackData.Flags, TrackData.SortOrder });
 			}
 		}
 		else
diff --git a/Engine/Source/Runtime/MovieSceneTracks/Private/Sections/MovieSceneCameraCutSection.cpp b/Engine/Source/Runtime/MovieSceneTracks/Private/Sections/MovieSceneCameraCutSection.cpp
--- a/Engine/Source/Runtime/MovieSceneTracks/Private/Sections/MovieSceneCameraCutSection.cpp
+++ b/Engine/Source/Runtime/MovieSceneTracks/Private/Sections/MovieSceneCameraCutSection.cpp
@@ -110,7 +110,7 @@ void UMovieSceneCameraCutSection::ImportEntityImpl(UMovieSceneEntitySystemLinker
 {
 	using namespace UE::MovieScene;
 
-	FMovieSceneTrackInstanceComponent TrackInstance { this, UMovieSceneCameraCutTrackInstance::StaticClass() };
+	FMovieSceneTrackInstanceComponent TrackInstance { decltype(FMovieSceneTrackInstanceComponent::Owner)(this), UMovieSceneCameraCutTrackInstance::StaticClass() };
 
 	OutImportedEntity->AddBuilder(
 		FEntityBuilder()
diff --git a/Engine/Source/Runtime/MovieSceneTracks/Private/Sections/MovieSceneLevelVisibilitySection.cpp b/Engine/Source/Runtime/MovieSceneTracks/Private/Sections/MovieSceneLevelVisibilitySection.cpp
--- a/Engine/Source/Runtime/MovieSceneTracks/Private/Sections/MovieSceneLevelVisibilitySection.cpp
+++ b/Engine/Source/Runtime/MovieSceneTracks/Private/Sections/MovieSceneLevelVisibilitySection.cpp
@@ -27,7 +27,7 @@ void UMovieSceneLevelVisibilitySection::ImportEntityImpl(UMovieSceneEntitySystem
 {
 	using namespace UE::MovieScene;
 
-	FLevelVisibilityComponentData LevelVisibilityData{ this };
+	FLevelVisibilityComponentData LevelVisibilityData{decltype(FLevelVisibilityComponentData::Section)(this) };
 
 	OutImportedEntity->AddBuilder(
 		FEntityBuilder()
diff --git a/Engine/Source/Runtime/SlateCore/Private/Fonts/CompositeFont.cpp b/Engine/Source/Runtime/SlateCore/Private/Fonts/CompositeFont.cpp
--- a/Engine/Source/Runtime/SlateCore/Private/Fonts/CompositeFont.cpp
+++ b/Engine/Source/Runtime/SlateCore/Private/Fonts/CompositeFont.cpp
@@ -214,7 +214,7 @@ bool FFontData::Serialize(FArchive& Ar)
 	if (bIsCooked)
 	{
 		// Cooked data uses a more compact format
-		UObject* LocalFontFaceAsset = const_cast<UObject*>(FontFaceAsset);
+		UObject* LocalFontFaceAsset = const_cast<UObject*>(ToRawPtr(FontFaceAsset));
 		Ar << LocalFontFaceAsset;
 		FontFaceAsset = LocalFontFaceAsset;
 		
