Reduce memory usage by allowing meshdescription and rawmesh bulkdata to be cleaned up if possible
Also add protection so meshdescription can be read from multiple threads in editor

 - Loading Reverb's P_World with empty DDC
    - 319GB to 242GB peak private byte
    - 166GB to 80GB permanent memory usage after map load

diff --git a/Engine/Source/Runtime/CoreUObject/Private/Serialization/BulkData.cpp b/Engine/Source/Runtime/CoreUObject/Private/Serialization/BulkData.cpp
--- a/Engine/Source/Runtime/CoreUObject/Private/Serialization/BulkData.cpp
+++ b/Engine/Source/Runtime/CoreUObject/Private/Serialization/BulkData.cpp
@@ -699,7 +699,7 @@ public:
 bool FUntypedBulkData::LoadBulkDataWithFileReader()
 {
 #if WITH_EDITOR
-	if (!BulkData && GIsEditor && !GEventDrivenLoaderEnabled && !SerializeFuture.IsValid())
+	if (!BulkData && CanLoadBulkDataWithFileReader() && !SerializeFuture.IsValid())
 	{
 		SerializeFuture = TFuture<bool>(TSharedPtr<TFutureState<bool>, ESPMode::ThreadSafe>(new FStateComplete([=]() 
 		{ 
@@ -712,6 +712,15 @@ bool FUntypedBulkData::LoadBulkDataWithFileReader()
 	return false;
 }
 
+bool FUntypedBulkData::CanLoadBulkDataWithFileReader() const
+{
+#if WITH_EDITOR
+	return !Filename.IsEmpty() && GIsEditor && !GEventDrivenLoaderEnabled;
+#else
+	return false;
+#endif
+}
+
 /**
  * Forces the bulk data to be resident in memory and detaches the archive.
  */
@@ -1205,7 +1214,7 @@ void FUntypedBulkData::Serialize( FArchive& Ar, UObject* Owner, int32 Idx, bool
 						// if the payload is NOT stored inline ...
 						if (BulkDataFlags & BULKDATA_PayloadInSeperateFile)
 						{
-							// open seperate bulk data file
+							// open separate bulk data file
 							UE_CLOG(GEventDrivenLoaderEnabled, LogSerialization, Error, TEXT("Attempt to sync load bulk data with EDL enabled (separate file). This is not desireable. File %s"), *Filename);
 
 							if (GEventDrivenLoaderEnabled && (Filename.EndsWith(TEXT(".uasset")) || Filename.EndsWith(TEXT(".umap"))))
diff --git a/Engine/Source/Runtime/CoreUObject/Public/Serialization/BulkData.h b/Engine/Source/Runtime/CoreUObject/Public/Serialization/BulkData.h
--- a/Engine/Source/Runtime/CoreUObject/Public/Serialization/BulkData.h
+++ b/Engine/Source/Runtime/CoreUObject/Public/Serialization/BulkData.h
@@ -571,6 +571,12 @@ public:
 	 */
 	bool LoadBulkDataWithFileReader();
 
+	/**
+	 * Test if it is possible to load the bulk data using a file reader, even when no archive is attached to the bulk data.
+	 * @return Whether the operation is allowed.
+	 */
+	bool CanLoadBulkDataWithFileReader() const;
+
 	/**
 	 * Forces the bulk data to be resident in memory and detaches the archive.
 	 */
diff --git a/Engine/Source/Runtime/MeshDescription/Private/MeshDescription.cpp b/Engine/Source/Runtime/MeshDescription/Private/MeshDescription.cpp
--- a/Engine/Source/Runtime/MeshDescription/Private/MeshDescription.cpp
+++ b/Engine/Source/Runtime/MeshDescription/Private/MeshDescription.cpp
@@ -11,6 +11,9 @@
 #include "UObject/UE5MainStreamObjectVersion.h"
 #include "Virtualization/VirtualizationUtilities.h"
 
+#if WITH_EDITOR
+#include "Misc/ScopeRWLock.h"
+#endif
 
 FName FMeshDescription::VerticesName("Vertices");
 FName FMeshDescription::VertexInstancesName("VertexInstances");
@@ -2040,7 +2043,6 @@ void FMeshDescriptionBulkData::SaveMeshDescription( FMeshDescription& MeshDescri
 	bBulkDataUpdated = true;
 }
 
-
 void FMeshDescriptionBulkData::LoadMeshDescription( FMeshDescription& MeshDescription )
 {
 	MeshDescription.Empty();
@@ -2049,25 +2051,47 @@ void FMeshDescriptionBulkData::LoadMeshDescription( FMeshDescription& MeshDescri
 	{
 		// Get a lock on the bulk data and read it into the mesh description
 		{
-			const bool bIsPersistent = true;
+#if WITH_EDITOR
+			// A lock is required so we can clone the mesh description from multiple threads
+			FRWScopeLock ScopeLock(BulkDataLock, SLT_Write);
+
+			// This allows any thread to be able to clone a mesh description directly
+			// from disk so we can unload bulk data from memory.
+			bool bHasBeenLoadedFromFileReader = false;
+			if (BulkData.IsAsyncLoadingComplete() && !BulkData.IsBulkDataLoaded())
+			{
+				// This will return false under -game mode for now because we're not allowed to load bulk data outside of EDL.
+				bHasBeenLoadedFromFileReader = BulkData.LoadBulkDataWithFileReader();
+			}
+#endif
+			// This is in a scope because the FBulkDataReader need to be destroyed in order
+			// to unlock the BulkData and allow UnloadBulkData to actually do its job.
+			{
+				const bool bIsPersistent = true;
 #if UE_USE_VIRTUALBULKDATA
-			FVirtualizedBulkDataReader Ar(BulkData, bIsPersistent);
+				FVirtualizedBulkDataReader Ar(BulkData, bIsPersistent);
 #else
-			FBulkDataReader Ar(BulkData, bIsPersistent);
+				FBulkDataReader Ar(BulkData, bIsPersistent);
 #endif //UE_USE_VIRTUALBULKDATA
 
-			// Propagate the custom version information from the package to the bulk data, so that the MeshDescription
-			// is serialized with the same versioning.
-			Ar.SetCustomVersions( CustomVersions );
-			Ar << MeshDescription;
-		}
-		// Unlock bulk data when we leave scope
+				// Propagate the custom version information from the package to the bulk data, so that the MeshDescription
+				// is serialized with the same versioning.
+				Ar.SetCustomVersions( CustomVersions );
+				Ar << MeshDescription;
+			} // Unlock bulk data when we leave scope
 
-		// Throw away the bulk data allocation as we don't need it now we have its contents as a FMeshDescription
-		// @todo: revisit this
 #if UE_USE_VIRTUALBULKDATA
-		BulkData.UnloadData();
+			BulkData.UnloadData();
+#elif WITH_EDITOR
+			// Throw away the bulk data allocation only in the case we can safely reload it from disk
+			// and if BulkData.LoadBulkDataWithFileReader() is allowed to work from any thread.
+			// This saves a significant amount of memory during map loading of Nanite Meshes.
+			if (bHasBeenLoadedFromFileReader)
+			{
+				verify(BulkData.UnloadBulkData());
+			}
 #endif //UE_USE_VIRTUALBULKDATA
+		}
 	}
 }
 
diff --git a/Engine/Source/Runtime/MeshDescription/Public/MeshDescription.h b/Engine/Source/Runtime/MeshDescription/Public/MeshDescription.h
--- a/Engine/Source/Runtime/MeshDescription/Public/MeshDescription.h
+++ b/Engine/Source/Runtime/MeshDescription/Public/MeshDescription.h
@@ -1331,6 +1331,10 @@ public:
 	void UseHashAsGuid();
 
 private:
+#if WITH_EDITOR
+	/** Protects simultaneous access to BulkData */
+	FRWLock       BulkDataLock;
+#endif
 	/** Internally store bulk data as bytes */
 	TChooseClass<UE_USE_VIRTUALBULKDATA, FByteVirtualizedBulkData, FByteBulkData>::Result BulkData;
 
diff --git a/Engine/Source/Runtime/RawMesh/Private/RawMesh.cpp b/Engine/Source/Runtime/RawMesh/Private/RawMesh.cpp
--- a/Engine/Source/Runtime/RawMesh/Private/RawMesh.cpp
+++ b/Engine/Source/Runtime/RawMesh/Private/RawMesh.cpp
@@ -5,6 +5,7 @@
 #include "UObject/Object.h"
 #include "Misc/SecureHash.h"
 #include "Modules/ModuleManager.h"
+#include "Serialization/BulkDataReader.h"
 
 IMPLEMENT_MODULE(FDefaultModuleImpl, RawMesh);
 
@@ -268,12 +269,37 @@ void FRawMeshBulkData::LoadRawMesh(FRawMesh& OutMesh)
 	OutMesh.Empty();
 	if (BulkData.GetElementCount() > 0)
 	{
-		FBufferReader Ar(
-			BulkData.Lock(LOCK_READ_ONLY), BulkData.GetElementCount(),
-			/*bInFreeOnClose=*/ false, /*bIsPersistent=*/ true
-			);
-		Ar << OutMesh;
-		BulkData.Unlock();
+#if WITH_EDITOR
+		// A lock is required so we can safely load the raw data from multiple threads
+		FRWScopeLock ScopeLock(BulkDataLock, SLT_Write);
+
+		// This allows any thread to be able to deserialize from the RawMesh directly
+		// from disk so we can unload bulk data from memory.
+		bool bHasBeenLoadedFromFileReader = false;
+		if (BulkData.IsAsyncLoadingComplete() && !BulkData.IsBulkDataLoaded())
+		{
+			// This can't be called in -game mode because we're not allowed to load bulk data outside of EDL.
+			bHasBeenLoadedFromFileReader = BulkData.LoadBulkDataWithFileReader();
+		}
+#endif
+
+		// This is in a scope because the FBulkDataReader need to be destroyed in order
+		// to unlock the BulkData and allow UnloadBulkData to actually do its job.
+		{
+			const bool bIsPersistent = true;
+			FBulkDataReader Ar(BulkData, bIsPersistent);
+			Ar << OutMesh;
+		}
+
+#if WITH_EDITOR
+		// Throw away the bulk data allocation only in the case we can safely reload it from disk
+		// and if BulkData.LoadBulkDataWithFileReader() is allowed to work from any thread.
+		// This saves a significant amount of memory during map loading of Nanite Meshes.
+		if (bHasBeenLoadedFromFileReader)
+		{
+			verify(BulkData.UnloadBulkData());
+		}
+#endif
 	}
 }
 
diff --git a/Engine/Source/Runtime/RawMesh/Public/RawMesh.h b/Engine/Source/Runtime/RawMesh/Public/RawMesh.h
--- a/Engine/Source/Runtime/RawMesh/Public/RawMesh.h
+++ b/Engine/Source/Runtime/RawMesh/Public/RawMesh.h
@@ -99,6 +99,10 @@ RAWMESH_API FArchive& operator<<(FArchive& Ar, FRawMesh& RawMesh);
  */
 class FRawMeshBulkData
 {
+#if WITH_EDITOR
+	/** Protects simultaneous access to BulkData */
+	FRWLock       BulkDataLock;
+#endif
 	/** Internally store bulk data as bytes. */
 	FByteBulkData BulkData;
 	/** GUID associated with the data stored herein. */
