Added unbounded arrays as a new attribute type for mesh attributes. They are registered with RegisterAttribute<T[]>, and accessed with GetAttributesRef<TAttributeArray<T>>.
Removed attribute views.
Refactored MeshAttributeArray implementation.
Added a new GetArrayView() method for all types of TMeshAttributeRefs which can be used generically.

diff --git a/Engine/Source/Runtime/MeshDescription/Private/MeshDescription.cpp b/Engine/Source/Runtime/MeshDescription/Private/MeshDescription.cpp
--- a/Engine/Source/Runtime/MeshDescription/Private/MeshDescription.cpp
+++ b/Engine/Source/Runtime/MeshDescription/Private/MeshDescription.cpp
@@ -338,31 +338,59 @@ struct FMeshPolygonGroup_Legacy
 
 
 template <typename T>
-void FixAttributesSize(int32 ExpectedNum, TAttributesSet<T>& AttributesSet)
+struct FFixAttributesSizeHelper
 {
-	// Ensure that the attribute set is the same size as the mesh element array they describe
-	// If there are extra elements, and they are not set to trivial defaults, this is an error.
-	bool bAllDefault = true;
-	AttributesSet.ForEach([ExpectedNum, &bAllDefault](const FName AttributeName, auto AttributesConstRef)
+	explicit FFixAttributesSizeHelper(int32 InExpectedNum)
+		: ExpectedNum(InExpectedNum),
+		  bAllDefault(true)
+	{}
+
+	template <typename U>
+	void operator()(const FName AttributeName, TMeshAttributesRef<T, TArrayView<U>> AttributeArrayRef)
+	{
+		// Not expecting arrays in legacy attributes
+		check(false);
+	}
+
+	template <typename U>
+	void operator()(const FName AttributeName, TMeshAttributesRef<T, TArrayAttribute<U>> AttributeArrayRef)
+	{
+		// Not expecting arrays in legacy attributes
+		check(false);
+	}
+
+	template <typename U>
+	void operator()(const FName AttributeName, TMeshAttributesRef<T, U> AttributeArrayRef)
+	{
+		if (bAllDefault)
 		{
-			if (bAllDefault)
+			for (int32 Channel = 0; Channel < AttributeArrayRef.GetNumChannels(); Channel++)
 			{
-				for (int32 Channel = 0; Channel < AttributesConstRef.GetNumChannels(); Channel++)
+				for (int32 Index = ExpectedNum; Index < AttributeArrayRef.GetNumElements(); Index++)
 				{
-					for (int32 Index = ExpectedNum; Index < AttributesConstRef.GetNumElements(); Index++)
+					if (AttributeArrayRef.Get(T(Index), Channel) != AttributeArrayRef.GetDefaultValue())
 					{
-						if (AttributesConstRef.Get(T(Index), Channel) != AttributesConstRef.GetDefaultValue())
-						{
-							bAllDefault = false;
-							break;
-						}
+						bAllDefault = false;
+						return;
 					}
 				}
 			}
 		}
-	);
+	}
+
+	int32 ExpectedNum;
+	bool bAllDefault;
+};
+
 
-	check(bAllDefault);	// If this fires, something is very wrong with the legacy asset
+template <typename T>
+void FixAttributesSize(int32 ExpectedNum, TAttributesSet<T>& AttributesSet)
+{
+	// Ensure that the attribute set is the same size as the mesh element array they describe
+	// If there are extra elements, and they are not set to trivial defaults, this is an error.
+	FFixAttributesSizeHelper<T> Helper(ExpectedNum);
+	AttributesSet.ForEach(Helper);
+	check(Helper.bAllDefault);	// If this fires, something is very wrong with the legacy asset
 	AttributesSet.SetNumElements(ExpectedNum);
 }
 
diff --git a/Engine/Source/Runtime/MeshDescription/Public/AttributeArrayContainer.h b/Engine/Source/Runtime/MeshDescription/Public/AttributeArrayContainer.h
new file mode 100644
--- /dev/null
+++ b/Engine/Source/Runtime/MeshDescription/Public/AttributeArrayContainer.h
@@ -0,0 +1,403 @@
+// Copyright Epic Games, Inc. All Rights Reserved.
+
+#pragma once
+
+#include "CoreMinimal.h"
+#include "Containers/StaticArray.h"
+#include "Algo/Rotate.h"
+
+
+template <typename AttributeType>
+class TAttributeArrayContainer
+{
+public:
+	TAttributeArrayContainer() = default;
+
+	explicit TAttributeArrayContainer(const AttributeType& InDefaultValue)
+		: DefaultValue(InDefaultValue)
+	{}
+
+	/** Return size of container */
+	int32 Num() const { return NumElements; }
+
+	/** Initializes the array to the given size with the default value */
+	void Initialize(const int32 ElementCount, const AttributeType& Default)
+	{
+		// For unbounded arrays, the default value is an empty subarray.
+		check(ElementCount >= 0);
+		int32 NumChunks = (ElementCount + ChunkSize - 1) >> ChunkBits;
+		Chunks.Empty(NumChunks);
+		Chunks.SetNum(NumChunks);
+		NumElements = ElementCount;
+	}
+
+	/** Sets the number of elements, each element itself being a subarray of items of type AttributeType. */
+	void SetNum(const int32 ElementCount, const AttributeType& Default)
+	{
+		check(ElementCount >= 0);
+
+		int32 NumChunks = (ElementCount + ChunkSize - 1) >> ChunkBits;
+		Chunks.SetNum(NumChunks);
+
+		if (ElementCount < NumElements)
+		{
+			// Case where we're shrinking the unbounded array
+			int32 IndexInLastChunk = ElementCount & ChunkMask;
+			if (IndexInLastChunk > 0)
+			{
+				FChunk& LastChunk = Chunks.Last();
+				int32 LastIndex = LastChunk.StartIndex[IndexInLastChunk - 1] + LastChunk.MaxCount[IndexInLastChunk - 1];
+				LastChunk.Data.SetNum(LastIndex);
+				for (; IndexInLastChunk < ChunkSize; ++IndexInLastChunk)
+				{
+					LastChunk.StartIndex[IndexInLastChunk] = LastIndex;
+					LastChunk.Count[IndexInLastChunk] = 0;
+					LastChunk.MaxCount[IndexInLastChunk] = 0;
+				}
+			}
+		}
+		else
+		{
+			// If we're growing the unbounded array, there's nothing to do;
+			// the excess indices in the old last chunk will already be set up with zero length, pointing at the off-the-end element of the data array
+		}
+
+		NumElements = ElementCount;
+	}
+
+	uint32 GetHash(uint32 Crc = 0) const
+	{
+		for (const FChunk& Chunk : Chunks)
+		{
+			Crc = FCrc::MemCrc32(Chunk.Data.GetData(), Chunk.Data.Num() * sizeof(AttributeType), Crc);
+		}
+		return Crc;
+	}
+
+	/** Expands the array if necessary so that the passed element index is valid. Newly created elements will be assigned the default value. */
+	void Insert(const int32 Index, const AttributeType& Default)
+	{
+		check(Index >= 0);
+
+		int32 EndIndex = Index + 1;
+		if (EndIndex > NumElements)
+		{
+			int32 NumChunks = (EndIndex + ChunkSize - 1) >> ChunkBits;
+			Chunks.SetNum(NumChunks);
+		}
+
+		NumElements = EndIndex;
+	}
+
+	/** Fills the index with the default value */
+	void SetToDefault(const int32 Index, const AttributeType& Default)
+	{
+		checkSlow(Index >= 0 && Index < NumElements);
+		int32 ChunkIndex = Index >> ChunkBits;
+		int32 ElementIndex = Index & ChunkMask;
+
+		// Default value is an empty array. We do not compact the allocations when shrinking the data.
+		Chunks[ChunkIndex].Count[ElementIndex] = 0;
+	}
+
+	/** Remaps elements according to the passed remapping table */
+	void Remap(const TSparseArray<int32>& IndexRemap, const AttributeType& Default);
+
+	friend FArchive& operator <<(FArchive& Ar, TAttributeArrayContainer<AttributeType>& Array)
+	{
+		Ar << Array.Chunks;
+		Ar << Array.NumElements;
+		Ar << Array.DefaultValue;
+		return Ar;
+	}
+
+	/** Gets the attribute array at the given index as a TArrayView */
+	TArrayView<AttributeType> Get(int32 Index)
+	{
+		checkSlow(Index >= 0 && Index < NumElements);
+		int32 ChunkIndex = Index >> ChunkBits;
+		int32 ElementIndex = Index & ChunkMask;
+		FChunk& Chunk = Chunks[ChunkIndex];
+		return TArrayView<AttributeType>(Chunk.Data.GetData() + Chunk.StartIndex[ElementIndex], Chunk.Count[ElementIndex]);
+	}
+
+	/** Gets the attribute array at the given index as a TArrayView */
+	TArrayView<const AttributeType> Get(int32 Index) const
+	{
+		checkSlow(Index >= 0 && Index < NumElements);
+		int32 ChunkIndex = Index >> ChunkBits;
+		int32 ElementIndex = Index & ChunkMask;
+		const FChunk& Chunk = Chunks[ChunkIndex];
+		return TArrayView<const AttributeType>(Chunk.Data.GetData() + Chunk.StartIndex[ElementIndex], Chunk.Count[ElementIndex]);
+	}
+
+	/** Sets the attribute array at the given index to the given TArrayView */
+	void Set(int32 Index, TArrayView<const AttributeType> Value)
+	{
+		TArrayView<AttributeType> Element = SetElementCount(Index, Value.Num(), false);
+
+		for (int32 I = 0; I < Value.Num(); I++)
+		{
+			Element[I] = Value[I];
+		}
+	}
+
+	/** Sets the given attribute array element to have the given number of subarray elements. */
+	TArrayView<AttributeType> SetElementCount(int32 Index, int32 Size, bool bSetDefault = true)
+	{
+		checkSlow(Index >= 0 && Index < NumElements);
+		int32 ChunkIndex = Index >> ChunkBits;
+		int32 ElementIndex = Index & ChunkMask;
+		FChunk& Chunk = Chunks[ChunkIndex];
+
+		AttributeType* BasePtr = Chunk.Data.GetData();
+
+		int32 Extra = Size - Chunk.MaxCount[ElementIndex];
+		if (Extra > 0)
+		{
+			// If we are requesting that MaxCount grows for this element, we have to insert new data elements,
+			// and adjust subsequent start indices accordingly.
+			Chunk.Data.SetNum(Chunk.Data.Num() + Extra);
+
+			AttributeType* StartPtr = BasePtr + Chunk.StartIndex[ElementIndex] + Extra;
+			AttributeType* EndPtr = BasePtr + Chunk.Data.Num();
+			FMemory::Memmove(StartPtr, StartPtr - Extra, (EndPtr - StartPtr) * sizeof(AttributeType));
+
+			Chunk.MaxCount[ElementIndex] = Size;
+			for (int32 I = ElementIndex + 1; I < ChunkSize; I++)
+			{
+				Chunk.StartIndex[I] += Extra;
+			}
+		}
+
+		if (bSetDefault)
+		{
+			// Set any newly created members to the default value
+			for (int32 I = Chunk.Count[ElementIndex]; I < Chunk.MaxCount[ElementIndex]; I++)
+			{
+				Chunk.Data[Chunk.StartIndex[ElementIndex] + I] = DefaultValue;
+			}
+		}
+
+		Chunk.Count[ElementIndex] = Size;
+
+		return TArrayView<AttributeType>(BasePtr + Chunk.StartIndex[ElementIndex], Chunk.Count[ElementIndex]);
+	}
+
+	TArrayView<AttributeType> InsertIntoElement(int32 Index, int32 SubArrayIndex, int32 InsertCount = 1)
+	{
+		checkSlow(Index >= 0 && Index < NumElements);
+		int32 ChunkIndex = Index >> ChunkBits;
+		int32 ElementIndex = Index & ChunkMask;
+		FChunk& Chunk = Chunks[ChunkIndex];
+
+		int32 CurrentCount = Chunk.Count[ElementIndex];
+		checkSlow(SubArrayIndex >= 0 && SubArrayIndex <= CurrentCount);
+		int32 NewCount = CurrentCount + InsertCount;
+
+		// Resize the allocated space for this element if necessary
+		TArrayView<AttributeType> Element = SetElementCount(ElementIndex, NewCount, true);
+
+		// Insert the new subarray elements
+		AttributeType* SpanStart = Element.GetData() + SubArrayIndex;
+		int32 SpanSize = NewCount - SubArrayIndex;
+		Algo::Rotate(TArrayView<AttributeType>(SpanStart, SpanSize), SpanSize - InsertCount);
+
+		return Element;
+	}
+
+	TArrayView<AttributeType> RemoveFromElement(int32 ElementIndex, int32 SubArrayIndex, int32 Count = 1)
+	{
+		TArrayView<AttributeType> Element = Get(ElementIndex);
+		for (int32 I = SubArrayIndex + Count; I < Element.Num(); I++)
+		{
+			Element[I - Count] = Element[I];
+		}
+
+		return SetElementCount(ElementIndex, Element.Num() - Count, false);
+	}
+
+
+private:
+
+	static const int32 ChunkBits = 8;
+	static const int32 ChunkSize = (1 << ChunkBits);
+	static const int32 ChunkMask = ChunkSize - 1;
+
+	struct FChunk
+	{
+		FChunk()
+		{
+			Data.Reserve(ChunkSize);
+		}
+
+		friend FArchive& operator <<(FArchive& Ar, FChunk& Chunk)
+		{
+			Ar << Chunk.Data;
+			Ar << Chunk.StartIndex;
+			Ar << Chunk.Count;
+			Ar << Chunk.MaxCount;
+			return Ar;
+		}
+
+		// All the data for each element in the chunk, packed contiguously
+		TArray<AttributeType> Data;
+
+		// Start, count and allocated count in the Data array for each element in the chunk.
+		// Arranged as SoA for cache optimization, since the most frequent operation is
+		// adding a fixed amount to all the start indices when a value is inserted.
+		TStaticArray<int32, ChunkSize> StartIndex = TStaticArray<int32, ChunkSize>(0);
+		TStaticArray<int32, ChunkSize> Count = TStaticArray<int32, ChunkSize>(0);
+		TStaticArray<int32, ChunkSize> MaxCount = TStaticArray<int32, ChunkSize>(0);
+	};
+
+	TArray<FChunk> Chunks;
+	int32 NumElements = 0;
+	AttributeType DefaultValue;
+};
+
+
+
+template <typename AttributeType>
+void TAttributeArrayContainer<AttributeType>::Remap(const TSparseArray<int32>& IndexRemap, const AttributeType& Default)
+{
+	TAttributeArrayContainer<AttributeType> NewAttributeArray(Default);
+
+	for (typename TSparseArray<int32>::TConstIterator It(IndexRemap); It; ++It)
+	{
+		const int32 OldElementIndex = It.GetIndex();
+		const int32 NewElementIndex = IndexRemap[OldElementIndex];
+
+		NewAttributeArray.Insert(NewElementIndex, Default);
+		NewAttributeArray.Set(NewElementIndex, Get(OldElementIndex));
+	}
+
+	*this = MoveTemp(NewAttributeArray);
+}
+
+
+/**
+ * Proxy object which fields access to an unbounded array attribute container.
+ */
+template <typename AttributeType>
+class TArrayAttribute
+{
+	using ArrayType = typename TCopyQualifiersFromTo<AttributeType, TAttributeArrayContainer<typename TRemoveCV<AttributeType>::Type>>::Type;
+
+public:
+	explicit TArrayAttribute(ArrayType& InArray, int32 InIndex)
+		: Array(InArray),
+		  Index(InIndex)
+	{}
+
+	/**
+	 * Helper function for returning a typed pointer to the first array attribute entry.
+	 */
+	AttributeType* GetData() const
+	{
+		return Array.Get(Index).GetData();
+	}
+
+	/**
+	 * Tests if index is valid, i.e. than or equal to zero, and less than the number of elements in the array attribute.
+	 *
+	 * @param Index Index to test.
+	 *
+	 * @returns True if index is valid. False otherwise.
+	 */
+	bool IsValidIndex(int32 ArrayAttributeIndex) const
+	{
+		return Array.Get(Index).IsValidIndex(ArrayAttributeIndex);
+	}
+
+	/**
+	 * Returns true if the array attribute is empty and contains no elements. 
+	 *
+	 * @returns True if the array attribute is empty.
+	 */
+	bool IsEmpty() const
+	{
+		return Array.Get(Index).IsEmpty();
+	}
+
+	/**
+	 * Returns number of elements in the array attribute.
+	 *
+	 * @returns Number of elements in array attribute.
+	 */
+	int32 Num() const
+	{
+		return Array.Get(Index).Num();
+	}
+
+	/**
+	 * Array bracket operator. Returns reference to array attribute element at given index.
+	 *
+	 * @returns Reference to indexed element.
+	 */
+	AttributeType& operator[](int32 ArrayAttributeIndex) const
+	{
+		return Array.Get(Index)[ArrayAttributeIndex];
+	}
+
+	/**
+	 * Returns n-th last element from the array attribute.
+	 *
+	 * @param IndexFromTheEnd (Optional) Index from the end of array.
+	 *                        Default is 0.
+	 *
+	 * @returns Reference to n-th last element from the array.
+	 */
+	AttributeType& Last(int32 IndexFromTheEnd = 0) const
+	{
+		return Array.Get(Index).Last();
+	}
+
+	/**
+	 * Sets the number of elements in the array attribute.
+	 */
+	void SetNum(int32 Num)
+	{
+		Array.SetElementCount(Index, Num, true);
+	}
+
+	/**
+	 * Inserts a number of elements in the array attribute
+	 */
+	void Insert(int32 StartIndex, int32 Count)
+	{
+		Array.InsertIntoElement(Index, StartIndex, Count);
+	}
+
+	/**
+	 * Removes a number of elements from the array attribute
+	 */
+	void Remove(int32 StartIndex, int32 Count)
+	{
+		Array.RemoveFromElement(Index, StartIndex, Count);
+	}
+
+	/**
+	 * Return a TArrayView representing this array attribute.
+	 */
+	TArrayView<AttributeType> ToArrayView() { return Array.Get(Index); }
+
+	/**
+	 * Implicitly coerce an array attribute to a TArrayView.
+	 */
+	operator TArrayView<AttributeType>() { return Array.Get(Index); }
+
+public:
+	/**
+	* DO NOT USE DIRECTLY
+	* STL-like iterators to enable range-based for loop support.
+	*/
+	FORCEINLINE AttributeType* begin() const { return GetData(); }
+	FORCEINLINE AttributeType* end  () const { return GetData() + Num(); }
+
+private:
+	ArrayType& Array;
+	int32 Index;
+};
+
+
diff --git a/Engine/Source/Runtime/MeshDescription/Public/MeshAttributeArray.h b/Engine/Source/Runtime/MeshDescription/Public/MeshAttributeArray.h
--- a/Engine/Source/Runtime/MeshDescription/Public/MeshAttributeArray.h
+++ b/Engine/Source/Runtime/MeshDescription/Public/MeshAttributeArray.h
@@ -6,6 +6,7 @@
 #include "Containers/ArrayView.h"
 #include "MeshTypes.h"
 #include "MeshElementRemappings.h"
+#include "AttributeArrayContainer.h"
 #include "Misc/TVariant.h"
 #include "Templates/CopyQualifiersFromTo.h"
 #include "UObject/EditorObjectVersion.h"
@@ -39,27 +40,6 @@ template <typename Tuple> struct TVariantFromTuple;
 template <typename... Ts> struct TVariantFromTuple<TTuple<Ts...>> { using Type = TVariant<FEmptyVariantState, Ts...>; };
 
 
-/**
- * Helper template which, given an array type, can break into its element type and number of elements.
- */
-template <typename T> struct TBreakArray;
-template <typename T, SIZE_T N> struct TBreakArray<T[N]> : TIntegralConstant<uint32, N> { using Type = T; };
-
-
-/**
- * Class which implements a function jump table to be automatically generated at compile time.
- * This is used by TAttributesSet to provide O(1) dispatch by attribute type at runtime.
- */
-template <typename FnType, uint32 Size>
-struct TJumpTable
-{
-	template <typename... T>
-	explicit constexpr TJumpTable( T... Ts ) : Fns{ Ts... } {}
-
-	FnType* Fns[Size];
-};
-
-
 /**
  * Traits class to specify which attribute types can be bulk serialized.
  */
@@ -105,13 +85,15 @@ public:
 
 	void SetNum(const int32 ElementCount, const AttributeType& Default)
 	{
-		if (Container.Num() >= ElementCount)
+		if (int32(ElementCount * Extent) < Container.Num())
 		{
-			Container.SetNum(ElementCount);
+			// Setting to a lower number; just truncate the container
+			Container.SetNum(ElementCount * Extent);
 		}
 		else
 		{
-			Initialize(ElementCount, Default);
+			// Setting to a higher number; grow the container, inserting default value to end.
+			Insert(ElementCount - 1, Default);
 		}
 	}
 
@@ -298,13 +280,13 @@ public:
 	virtual void Serialize(FArchive& Ar) = 0;
 	virtual void Remap(const TSparseArray<int32>& IndexRemap) = 0;
 
-//	UE_DEPRECATED(5.0, "Please use GetNumChannels().")
+	UE_DEPRECATED(5.0, "Please use GetNumChannels().")
 	virtual int32 GetNumIndices() const = 0;
-//	UE_DEPRECATED(5.0, "Please use SetNumChannels().")
+	UE_DEPRECATED(5.0, "Please use SetNumChannels().")
 	virtual void SetNumIndices(const int32 NumIndices) = 0;
-//	UE_DEPRECATED(5.0, "Please use InsertChannel().")
+	UE_DEPRECATED(5.0, "Please use InsertChannel().")
 	virtual void InsertIndex(const int32 Index) = 0;
-//	UE_DEPRECATED(5.0, "Please use RemoveChannel().")
+	UE_DEPRECATED(5.0, "Please use RemoveChannel().")
 	virtual void RemoveIndex(const int32 Index) = 0;
 
 	virtual int32 GetNumChannels() const = 0;
@@ -355,6 +337,8 @@ protected:
 template <typename AttributeType>
 class TMeshAttributeArraySet : public FMeshAttributeArraySetBase
 {
+	static_assert(!TIsArray<AttributeType>::Value, "TMeshAttributeArraySet must take a simple type.");
+
 	using Super = FMeshAttributeArraySetBase;
 
 public:
@@ -525,65 +509,304 @@ protected:
 };
 
 
-template <typename> struct TIsArrayView { enum { Value = false }; };
-template <typename T> struct TIsArrayView<TArrayView<T>> { enum { Value = true }; };
+/**
+* This is a type-specific attribute array, which is actually instanced in the attribute set.
+*/
+template <typename AttributeType>
+class TMeshUnboundedAttributeArraySet : public FMeshAttributeArraySetBase
+{
+	using Super = FMeshAttributeArraySetBase;
+
+public:
+	/** Constructors */
+	FORCEINLINE TMeshUnboundedAttributeArraySet()
+		: Super(TTupleIndex<AttributeType, AttributeTypes>::Value, EMeshAttributeFlags::None, 0, 0)
+	{}
+
+	FORCEINLINE explicit TMeshUnboundedAttributeArraySet(const int32 NumberOfChannels, const AttributeType& InDefaultValue, const EMeshAttributeFlags InFlags, const int32 InNumberOfElements)
+		: Super(TTupleIndex<AttributeType, AttributeTypes>::Value, InFlags, InNumberOfElements, 0),
+		DefaultValue(InDefaultValue)
+	{
+		SetNumChannels(NumberOfChannels);
+	}
+
+	/** Creates a copy of itself and returns a TUniquePtr to it */
+	virtual TUniquePtr<FMeshAttributeArraySetBase> Clone() const override
+	{
+		return MakeUnique<TMeshUnboundedAttributeArraySet>(*this);
+	}
+
+	/** Insert the element at the given index */
+	virtual void Insert(const int32 Index) override
+	{
+		for (TAttributeArrayContainer<AttributeType>& ArrayForChannel : ArrayForChannels)
+		{
+			ArrayForChannel.Insert(Index, DefaultValue);
+		}
+
+		NumElements = FMath::Max(NumElements, Index + 1);
+	}
+
+	/** Remove the element at the given index, replacing it with a default value */
+	virtual void Remove(const int32 Index) override
+	{
+		for (TAttributeArrayContainer<AttributeType>& ArrayForChannel : ArrayForChannels)
+		{
+			ArrayForChannel.SetToDefault(Index, DefaultValue);
+		}
+	}
+
+	/** Sets the number of elements to the exact number provided, and initializes them to the default value */
+	virtual void Initialize(const int32 Count) override
+	{
+		NumElements = Count;
+		for (TAttributeArrayContainer<AttributeType>& ArrayForChannel : ArrayForChannels)
+		{
+			ArrayForChannel.Initialize(Count, DefaultValue);
+		}
+	}
+
+	/** Sets the number of elements to the exact number provided, preserving existing elements if the number is bigger */
+	virtual void SetNumElements(const int32 Count) override
+	{
+		NumElements = Count;
+		for (TAttributeArrayContainer<AttributeType>& ArrayForChannel : ArrayForChannels)
+		{
+			ArrayForChannel.SetNum(Count, DefaultValue);
+		}
+	}
+
+	virtual uint32 GetHash() const override
+	{
+		uint32 CrcResult = 0;
+		for (const TAttributeArrayContainer<AttributeType>& ArrayForChannel : ArrayForChannels)
+		{
+			CrcResult = ArrayForChannel.GetHash(CrcResult);
+		}
+		return CrcResult;
+	}
+
+	/** Polymorphic serialization */
+	virtual void Serialize(FArchive& Ar) override
+	{
+		Ar << (*this);
+	}
+
+	/** Performs an element index remap according to the passed array */
+	virtual void Remap(const TSparseArray<int32>& IndexRemap) override
+	{
+		for (TAttributeArrayContainer<AttributeType>& ArrayForChannel : ArrayForChannels)
+		{
+			ArrayForChannel.Remap(IndexRemap, DefaultValue);
+			NumElements = ArrayForChannel.Num();
+		}
+	}
+
+	UE_DEPRECATED(5.0, "Please use GetNumChannels().")
+	virtual inline int32 GetNumIndices() const override { return GetNumChannels(); }
+
+	/** Return number of channels this attribute has */
+	virtual inline int32 GetNumChannels() const override { return ArrayForChannels.Num(); }
+
+	UE_DEPRECATED(5.0, "Please use SetNumChannels().")
+	virtual void SetNumIndices(const int32 NumIndices) override { SetNumChannels(NumIndices); }
+
+	/** Sets number of channels this attribute has */
+	virtual void SetNumChannels(const int32 NumChannels) override
+	{
+		if (NumChannels < ArrayForChannels.Num())
+		{
+			ArrayForChannels.SetNum(NumChannels);
+			return;
+		}
+
+		while (ArrayForChannels.Num() < NumChannels)
+		{
+			TAttributeArrayContainer<AttributeType>& Array = ArrayForChannels.Emplace_GetRef(DefaultValue);
+			Array.Initialize(NumElements, DefaultValue);
+		}
+	}
+
+	UE_DEPRECATED(5.0, "Please use InsertChannel().")
+	virtual void InsertIndex(const int32 Index) override
+	{
+		InsertChannel(Index);
+	}
+
+	/** Insert a new attribute channel */
+	virtual void InsertChannel(const int32 Index) override
+	{
+		TAttributeArrayContainer<AttributeType>& Array = ArrayForChannels.EmplaceAt_GetRef(Index, DefaultValue);
+		Array.Initialize(NumElements, DefaultValue);
+	}
+
+	UE_DEPRECATED(5.0, "Please use RemoveChannel().")
+	virtual void RemoveIndex(const int32 Index) override
+	{
+		RemoveChannel(Index);
+	}
+
+	/** Remove the channel at the given index */
+	virtual void RemoveChannel(const int32 Index) override
+	{
+		ArrayForChannels.RemoveAt(Index);
+	}
+
+	/** Return the TMeshAttributeArrayBase corresponding to the given attribute channel */
+	FORCEINLINE const TAttributeArrayContainer<AttributeType>& GetArrayForChannel( const int32 Index ) const { return ArrayForChannels[ Index ]; }
+	FORCEINLINE TAttributeArrayContainer<AttributeType>& GetArrayForChannel( const int32 Index ) { return ArrayForChannels[ Index ]; }
+
+	/** Return default value for this attribute type */
+	FORCEINLINE AttributeType GetDefaultValue() const { return DefaultValue; }
+
+	/** Serializer */
+	friend FArchive& operator<<(FArchive& Ar, TMeshUnboundedAttributeArraySet& AttributeArraySet)
+	{
+		Ar << AttributeArraySet.NumElements;
+		Ar << AttributeArraySet.ArrayForChannels;
+		Ar << AttributeArraySet.DefaultValue;
+		Ar << AttributeArraySet.Flags;
+
+		return Ar;
+	}
+
+protected:
+	/** An array of UnboundedArrays, one per channel */
+	TArray<TAttributeArrayContainer<AttributeType>, TInlineAllocator<1>> ArrayForChannels;
+
+	/** The default value for an attribute of this name */
+	AttributeType DefaultValue;
+};
+
+
+/**
+ * Define type traits for different kinds of mesh attributes.
+ * 
+ * There are three type of attributes:
+ * - simple values (T)
+ * - fixed size arrays of values (TArrayView<T>)
+ * - variable size arrays of values (TArrayAttribute<T>)
+ *
+ * Each of these corresponds to a different type of TMeshAttributesRef and TMeshAttributeArraySet.
+ */
+template <typename T>
+struct TMeshAttributesRefTypeBase
+{
+	using AttributeType = T;
+	using RealAttributeType = typename TChooseClass<TIsDerivedFrom<AttributeType, FElementID>::Value, int32, AttributeType>::Result;
+};
+
+template <typename T>
+struct TMeshAttributesRefType : TMeshAttributesRefTypeBase<T>
+{
+	static const uint32 MinExpectedExtent = 1;
+	static const uint32 MaxExpectedExtent = 1;
+	using RefType = T;
+	using ConstRefType = const T;
+	using NonConstRefType = typename TRemoveCV<T>::Type;
+	using SetType = TMeshAttributeArraySet<typename TRemoveCV<T>::Type>;
+	using ConstSetType = const TMeshAttributeArraySet<typename TRemoveCV<T>::Type>;
+};
+
+template <typename T>
+struct TMeshAttributesRefType<TArrayView<T>> : TMeshAttributesRefTypeBase<T>
+{
+	static const uint32 MinExpectedExtent = 1;
+	static const uint32 MaxExpectedExtent = 0xFFFFFFFF;
+	using RefType = TArrayView<T>;
+	using ConstRefType = TArrayView<const T>;
+	using NonConstRefType = TArrayView<typename TRemoveCV<T>::Type>;
+	using SetType = TMeshAttributeArraySet<typename TRemoveCV<T>::Type>;
+	using ConstSetType = const TMeshAttributeArraySet<typename TRemoveCV<T>::Type>;
+};
+
+template <typename T>
+struct TMeshAttributesRefType<TArrayAttribute<T>> : TMeshAttributesRefTypeBase<T>
+{
+	static const uint32 MinExpectedExtent = 0;
+	static const uint32 MaxExpectedExtent = 0;
+	using RefType = TArrayAttribute<T>;
+	using ConstRefType = TArrayAttribute<const T>;
+	using NonConstRefType = TArrayAttribute<typename TRemoveCV<T>::Type>;
+	using SetType = TMeshUnboundedAttributeArraySet<typename TRemoveCV<T>::Type>;
+	using ConstSetType = const TMeshUnboundedAttributeArraySet<typename TRemoveCV<T>::Type>;
+};
+
+
+/**
+ * Additional type traits for registering different attributes.
+ * When registering, we need to specify the attribute type with a concrete element count if necessary, i.e.
+ * 
+ * - simple values (T)
+ * - fixed size arrays of values (T[N])
+ * - variable size arrays of values (T[])
+ */
+template <typename T>
+struct TMeshAttributesRegisterType : TMeshAttributesRefType<T>
+{
+	static const uint32 Extent = 1;
+};
 
+template <typename T, SIZE_T N>
+struct TMeshAttributesRegisterType<T[N]> : TMeshAttributesRefType<TArrayView<T>>
+{
+	static const uint32 Extent = N;
+};
 
-template <typename> struct TBreakArrayView;
-template <typename T> struct TBreakArrayView<TArrayView<T>> { using Type = T; };
+template <typename T>
+struct TMeshAttributesRegisterType<T[]> : TMeshAttributesRefType<TArrayAttribute<T>>
+{
+	static const uint32 Extent = 0;
+};
 
 
-template <typename T, bool IsArray = TIsArrayView<T>::Value> struct TConstAttribute;
-template <typename T> struct TConstAttribute<T, false> { using Type = const T; };
-template <typename T> struct TConstAttribute<TArrayView<T>, true> { using Type = TArrayView<const T>; };
 
 /**
  * This is the class used to access attribute values.
  * It is a proxy object to a TMeshAttributeArraySet<> and should be passed by value.
  * It is valid for as long as the owning FMeshDescription exists.
  */
-template <typename ElementIDType, typename AttributeType, bool IsArray = TIsArrayView<AttributeType>::Value>
+template <typename ElementIDType, typename AttributeType>
 class TMeshAttributesRef;
 
-template <typename ElementIDType, typename AttributeType, bool IsArray = TIsArrayView<AttributeType>::Value>
-using TMeshAttributesConstRef = TMeshAttributesRef<ElementIDType, typename TConstAttribute<AttributeType>::Type, IsArray>;
+template <typename ElementIDType, typename AttributeType>
+using TMeshAttributesConstRef = TMeshAttributesRef<ElementIDType, typename TMeshAttributesRefType<AttributeType>::ConstRefType>;
 
-template <typename AttributeType, bool IsArray = TIsArrayView<AttributeType>::Value>
-using TMeshAttributesArray = TMeshAttributesRef<int32, AttributeType, IsArray>;
+template <typename AttributeType>
+using TMeshAttributesArray = TMeshAttributesRef<int32, AttributeType>;
 
-template <typename AttributeType, bool IsArray = TIsArrayView<AttributeType>::Value>
-using TMeshAttributesConstArray = TMeshAttributesRef<int32, typename TConstAttribute<AttributeType>::Type, IsArray>;
+template <typename AttributeType>
+using TMeshAttributesConstArray = TMeshAttributesRef<int32, typename TMeshAttributesRefType<AttributeType>::ConstRefType>;
 
 
 template <typename ElementIDType, typename AttributeType>
-class TMeshAttributesRef<ElementIDType, AttributeType, false>
+class TMeshAttributesRef
 {
-	template <typename T, typename U, bool IsArray> friend class TMeshAttributesRef;
+	template <typename T, typename U> friend class TMeshAttributesRef;
 
 public:
-	using Type = AttributeType;
 	using ArrayType = typename TCopyQualifiersFromTo<AttributeType, TMeshAttributeArraySet<typename TRemoveCV<AttributeType>::Type>>::Type;
 
 	/** Constructor taking a pointer to a TMeshAttributeArraySet */
-	FORCEINLINE explicit TMeshAttributesRef(ArrayType* InArrayPtr = nullptr)
+	FORCEINLINE explicit TMeshAttributesRef(ArrayType* InArrayPtr = nullptr, uint32 InExtent = 1)
 		: ArrayPtr(InArrayPtr)
 	{}
 
 	/** Implicitly construct a TMeshAttributesRef-to-const from a regular one */
 	template <typename T = AttributeType, typename TEnableIf<TIsSame<T, const T>::Value, int>::Type = 0>
-	FORCEINLINE TMeshAttributesRef(TMeshAttributesRef<ElementIDType, typename TRemoveCV<T>::Type, false> InRef)
+	FORCEINLINE TMeshAttributesRef(TMeshAttributesRef<ElementIDType, typename TRemoveCV<T>::Type> InRef)
 		: ArrayPtr(InRef.ArrayPtr)
 	{}
 
 	/** Implicitly construct a TMeshAttributesRef from a TMeshAttributesArray **/
 	template <typename T = ElementIDType, typename U = AttributeType, typename TEnableIf<!TIsSame<T, int32>::Value, int>::Type = 0>
-	FORCEINLINE TMeshAttributesRef(TMeshAttributesRef<int32, U, false> InRef)
+	FORCEINLINE TMeshAttributesRef(TMeshAttributesRef<int32, U> InRef)
 		: ArrayPtr(InRef.ArrayPtr)
 	{}
 
 	/** Implicitly construct a TMeshAttributesRef-to-const from a TMeshAttributesArray */
 	template <typename T = ElementIDType, typename U = AttributeType, typename TEnableIf<TIsSame<U, const U>::Value, int>::Type = 0, typename TEnableIf<!TIsSame<T, int32>::Value, int>::Type = 0>
-	FORCEINLINE TMeshAttributesRef(TMeshAttributesRef<int32, typename TRemoveCV<U>::Type, false> InRef)
+	FORCEINLINE TMeshAttributesRef(TMeshAttributesRef<int32, typename TRemoveCV<U>::Type> InRef)
 		: ArrayPtr(InRef.ArrayPtr)
 	{}
 
@@ -594,16 +817,9 @@ public:
 		return ArrayPtr->GetArrayForChannel(0).GetElementBase(ElementID.GetValue())[0];
 	}
 
-	/** Get the element with the given ID from channel 0 */
-	template <typename T = ElementIDType, typename TEnableIf<TIsDerivedFrom<T, FElementID>::Value, int>::Type = 0>
-	FORCEINLINE AttributeType Get(const ElementIDType ElementID) const
-	{
-		return ArrayPtr->GetArrayForChannel(0).GetElementBase(ElementID.GetValue())[0];
-	}
-
 	/** Get the element with the given ID and channel */
 	template <typename T = ElementIDType, typename TEnableIf<TIsDerivedFrom<T, FElementID>::Value, int>::Type = 0>
-	FORCEINLINE AttributeType Get(const ElementIDType ElementID, const int32 Channel) const
+	FORCEINLINE AttributeType Get(const ElementIDType ElementID, const int32 Channel = 0) const
 	{
 		return ArrayPtr->GetArrayForChannel(Channel).GetElementBase(ElementID.GetValue())[0];
 	}
@@ -613,32 +829,32 @@ public:
 		return ArrayPtr->GetArrayForChannel(0).GetElementBase(ElementIndex)[0];
 	}
 
-	FORCEINLINE AttributeType Get(int32 ElementIndex) const
+	FORCEINLINE AttributeType Get(int32 ElementIndex, const int32 Channel = 0) const
 	{
-		return ArrayPtr->GetArrayForChannel(0).GetElementBase(ElementIndex)[0];
+		return ArrayPtr->GetArrayForChannel(Channel).GetElementBase(ElementIndex)[0];
 	}
 
-	FORCEINLINE AttributeType Get(int32 ElementIndex, const int32 Index) const
+	FORCEINLINE TArrayView<AttributeType> GetArrayView(int32 ElementIndex, const int32 Channel = 0) const
 	{
-		return ArrayPtr->GetArrayForChannel(Index).GetElementBase(ElementIndex)[0];
+		return TArrayView<AttributeType>(ArrayPtr->GetArrayForChannel(Channel).GetElementBase(ElementIndex), 1);
 	}
 
-	FORCEINLINE TArrayView<Type> GetRawArray(const int32 AttributeChannel = 0) const
+	FORCEINLINE TArrayView<AttributeType> GetRawArray(const int32 AttributeChannel = 0) const
 	{
 		if (ArrayPtr == nullptr)
 		{
-			return TArrayView<Type>();
+			return TArrayView<AttributeType>();
 		}
 
-		Type* Element = ArrayPtr->GetArrayForChannel(AttributeChannel).GetElementBase(0);
-		return TArrayView<Type>(Element, GetNumElements());
+		AttributeType* Element = ArrayPtr->GetArrayForChannel(AttributeChannel).GetElementBase(0);
+		return TArrayView<AttributeType>(Element, GetNumElements());
 	}
 
 	/** Return whether the reference is valid or not */
 	FORCEINLINE bool IsValid() const { return (ArrayPtr != nullptr); }
 
 	/** Return default value for this attribute type */
-	FORCEINLINE Type GetDefaultValue() const { return ArrayPtr->GetDefaultValue(); }
+	FORCEINLINE AttributeType GetDefaultValue() const { return ArrayPtr->GetDefaultValue(); }
 
 	UE_DEPRECATED(5.0, "Please use GetNumChannels().")
 	FORCEINLINE int32 GetNumIndices() const
@@ -665,30 +881,42 @@ public:
 
 	/** Set the element with the given ID and index 0 to the provided value */
 	template <typename T = ElementIDType, typename TEnableIf<TIsDerivedFrom<T, FElementID>::Value, int>::Type = 0>
-	FORCEINLINE void Set(const ElementIDType ElementID, const Type& Value) const
+	FORCEINLINE void Set(const ElementIDType ElementID, const AttributeType& Value) const
 	{
 		ArrayPtr->GetArrayForChannel(0).GetElementBase(ElementID.GetValue())[0] = Value;
 	}
 
 	/** Set the element with the given ID and channel to the provided value */
 	template <typename T = ElementIDType, typename TEnableIf<TIsDerivedFrom<T, FElementID>::Value, int>::Type = 0>
-	FORCEINLINE void Set(const ElementIDType ElementID, const int32 Channel, const Type& Value) const
+	FORCEINLINE void Set(const ElementIDType ElementID, const int32 Channel, const AttributeType& Value) const
 	{
 		ArrayPtr->GetArrayForChannel(Channel).GetElementBase(ElementID.GetValue())[0] = Value;
 	}
 
-	FORCEINLINE void Set(int32 ElementIndex, const Type& Value) const
+	FORCEINLINE void Set(int32 ElementIndex, const AttributeType& Value) const
 	{
 		ArrayPtr->GetArrayForChannel(0).GetElementBase(ElementIndex)[0] = Value;
 	}
 
-	FORCEINLINE void Set(int32 ElementIndex, const int32 Channel, const Type& Value) const
+	FORCEINLINE void Set(int32 ElementIndex, const int32 Channel, const AttributeType& Value) const
 	{
 		ArrayPtr->GetArrayForChannel(Channel).GetElementBase(ElementIndex)[0] = Value;
 	}
 
-	/** Copies the given attribute array and index to this index */
-	void Copy(TMeshAttributesRef<ElementIDType, const AttributeType, false> Src, const int32 DestIndex = 0, const int32 SrcIndex = 0);
+	FORCEINLINE void SetArrayView(int32 ElementIndex, TArrayView<const AttributeType> Value) const
+	{
+		check(Value.Num() == 1);
+		ArrayPtr->GetArrayForChannel(0).GetElementBase(ElementIndex)[0] = Value[0];
+	}
+
+	FORCEINLINE void SetArrayView(int32 ElementIndex, const int32 Channel, TArrayView<const AttributeType> Value) const
+	{
+		check(Value.Num() == 1);
+		ArrayPtr->GetArrayForChannel(Channel).GetElementBase(ElementIndex)[0] = Value[0];
+	}
+
+	/** Copies the given attribute array and channel to this channel */
+	void Copy(TMeshAttributesRef<ElementIDType, const AttributeType> Src, const int32 DestChannel = 0, const int32 SrcChannel = 0);
 
 	UE_DEPRECATED(5.0, "Please use SetNumChannels().")
 	FORCEINLINE void SetNumIndices(const int32 NumChannels) const
@@ -731,11 +959,11 @@ protected:
 };
 
 template <typename ElementIDType, typename AttributeType>
-void TMeshAttributesRef<ElementIDType, AttributeType, false>::Copy(TMeshAttributesRef<ElementIDType, const AttributeType, false> Src, const int32 DestIndex, const int32 SrcIndex)
+void TMeshAttributesRef<ElementIDType, AttributeType>::Copy(TMeshAttributesRef<ElementIDType, const AttributeType> Src, const int32 DestChannel, const int32 SrcChannel)
 {
 	check(Src.IsValid());
-	const TMeshAttributeArrayBase<AttributeType>& SrcArray = Src->ArrayPtr->GetArrayForChannel(SrcIndex);
-	TMeshAttributeArrayBase<AttributeType>& DestArray = ArrayPtr->GetArrayForChannel(DestIndex);
+	const TMeshAttributeArrayBase<AttributeType>& SrcArray = Src->ArrayPtr->GetArrayForChannel(SrcChannel);
+	TMeshAttributeArrayBase<AttributeType>& DestArray = ArrayPtr->GetArrayForChannel(DestChannel);
 	const int32 Num = FMath::Min(SrcArray.Num(), DestArray.Num());
 	for (int32 Index = 0; Index < Num; Index++)
 	{
@@ -745,13 +973,12 @@ void TMeshAttributesRef<ElementIDType, AttributeType, false>::Copy(TMeshAttribut
 
 
 template <typename ElementIDType, typename AttributeType>
-class TMeshAttributesRef<ElementIDType, AttributeType, true>
+class TMeshAttributesRef<ElementIDType, TArrayView<AttributeType>>
 {
-	template <typename T, typename U, bool IsArray> friend class TMeshAttributesRef;
+	template <typename T, typename U> friend class TMeshAttributesRef;
 
 public:
-	using Type = typename TBreakArrayView<AttributeType>::Type;
-	using ArrayType = typename TCopyQualifiersFromTo<Type, TMeshAttributeArraySet<typename TRemoveCV<Type>::Type>>::Type;
+	using ArrayType = typename TCopyQualifiersFromTo<AttributeType, TMeshAttributeArraySet<typename TRemoveCV<AttributeType>::Type>>::Type;
 
 	/** Constructor taking a pointer to a TMeshAttributeArraySet */
 	FORCEINLINE explicit TMeshAttributesRef(ArrayType* InArrayPtr = nullptr, uint32 InExtent = 1)
@@ -760,22 +987,22 @@ public:
 	{}
 
 	/** Implicitly construct a TMeshAttributesRef-to-const from a regular one */
-	template <typename T = Type, typename TEnableIf<TIsSame<T, const T>::Value, int>::Type = 0>
-	FORCEINLINE TMeshAttributesRef(TMeshAttributesRef<ElementIDType, TArrayView<typename TRemoveCV<T>::Type>, true> InRef)
+	template <typename T = AttributeType, typename TEnableIf<TIsSame<T, const T>::Value, int>::Type = 0>
+	FORCEINLINE TMeshAttributesRef(TMeshAttributesRef<ElementIDType, TArrayView<typename TRemoveCV<T>::Type>> InRef)
 		: ArrayPtr(InRef.ArrayPtr),
 		  Extent(InRef.Extent)
 	{}
 
 	/** Implicitly construct a TMeshAttributesRef from a TMeshAttributesArray **/
 	template <typename T = ElementIDType, typename U = AttributeType, typename TEnableIf<!TIsSame<T, int32>::Value, int>::Type = 0>
-	FORCEINLINE TMeshAttributesRef(TMeshAttributesRef<int32, U, true> InRef)
+	FORCEINLINE TMeshAttributesRef(TMeshAttributesRef<int32, TArrayView<U>> InRef)
 		: ArrayPtr(InRef.ArrayPtr),
 		  Extent(InRef.Extent)
 	{}
 
 	/** Implicitly construct a TMeshAttributesRef-to-const from a TMeshAttributesArray */
-	template <typename T = ElementIDType, typename U = Type, typename TEnableIf<TIsSame<U, const U>::Value, int>::Type = 0, typename TEnableIf<!TIsSame<T, int32>::Value, int>::Type = 0>
-	FORCEINLINE TMeshAttributesRef(TMeshAttributesRef<int32, TArrayView<typename TRemoveCV<U>::Type>, true> InRef)
+	template <typename T = ElementIDType, typename U = AttributeType, typename TEnableIf<TIsSame<U, const U>::Value, int>::Type = 0, typename TEnableIf<!TIsSame<T, int32>::Value, int>::Type = 0>
+	FORCEINLINE TMeshAttributesRef(TMeshAttributesRef<int32, TArrayView<typename TRemoveCV<U>::Type>> InRef)
 		: ArrayPtr(InRef.ArrayPtr),
 		  Extent(InRef.Extent)
 	{}
@@ -783,62 +1010,53 @@ public:
 
 	/** Access elements from attribute channel 0 */
 	template <typename T = ElementIDType, typename TEnableIf<TIsDerivedFrom<T, FElementID>::Value, int>::Type = 0>
-	FORCEINLINE TArrayView<Type> operator[](const ElementIDType ElementID) const
-	{
-		Type* Element = ArrayPtr->GetArrayForChannel(0).GetElementBase(ElementID.GetValue());
-		return TArrayView<Type>(Element, Extent);
-	}
-
-	/** Get the element with the given ID from channel 0 */
-	template <typename T = ElementIDType, typename TEnableIf<TIsDerivedFrom<T, FElementID>::Value, int>::Type = 0>
-	FORCEINLINE TArrayView<Type> Get(const ElementIDType ElementID) const
+	FORCEINLINE TArrayView<AttributeType> operator[](const ElementIDType ElementID) const
 	{
-		Type* Element = ArrayPtr->GetArrayForChannel(0).GetElementBase(ElementID.GetValue());
-		return TArrayView<Type>(Element, Extent);
+		AttributeType* Element = ArrayPtr->GetArrayForChannel(0).GetElementBase(ElementID.GetValue());
+		return TArrayView<AttributeType>(Element, Extent);
 	}
 
 	/** Get the element with the given ID and channel */
 	template <typename T = ElementIDType, typename TEnableIf<TIsDerivedFrom<T, FElementID>::Value, int>::Type = 0>
-	FORCEINLINE TArrayView<Type> Get(const ElementIDType ElementID, const int32 Channel) const
+	FORCEINLINE TArrayView<AttributeType> Get(const ElementIDType ElementID, const int32 Channel = 0) const
 	{
-		Type* Element = ArrayPtr->GetArrayForChannel(Channel).GetElementBase(ElementID.GetValue());
-		return TArrayView<Type>(Element, Extent);
+		AttributeType* Element = ArrayPtr->GetArrayForChannel(Channel).GetElementBase(ElementID.GetValue());
+		return TArrayView<AttributeType>(Element, Extent);
 	}
 
-	FORCEINLINE TArrayView<Type> operator[](int32 ElementIndex) const
+	FORCEINLINE TArrayView<AttributeType> operator[](int32 ElementIndex) const
 	{
-		Type* Element = ArrayPtr->GetArrayForChannel(0).GetElementBase(ElementIndex);
-		return TArrayView<Type>(Element, Extent);
+		AttributeType* Element = ArrayPtr->GetArrayForChannel(0).GetElementBase(ElementIndex);
+		return TArrayView<AttributeType>(Element, Extent);
 	}
 
-	FORCEINLINE TArrayView<Type> Get(int32 ElementIndex) const
+	FORCEINLINE TArrayView<AttributeType> Get(int32 ElementIndex, const int32 Channel = 0) const
 	{
-		Type* Element = ArrayPtr->GetArrayForChannel(0).GetElementBase(ElementIndex);
-		return TArrayView<Type>(Element, Extent);
+		AttributeType* Element = ArrayPtr->GetArrayForChannel(Channel).GetElementBase(ElementIndex);
+		return TArrayView<AttributeType>(Element, Extent);
 	}
 
-	FORCEINLINE TArrayView<Type> Get(int32 ElementIndex, const int32 Channel) const
+	FORCEINLINE TArrayView<AttributeType> GetArrayView(int32 ElementIndex, const int32 Channel = 0) const
 	{
-		Type* Element = ArrayPtr->GetArrayForChannel(Channel).GetElementBase(ElementIndex);
-		return TArrayView<Type>(Element, Extent);
+		return TArrayView<AttributeType>(ArrayPtr->GetArrayForChannel(Channel).GetElementBase(ElementIndex), Extent);
 	}
 
-	FORCEINLINE TArrayView<Type> GetRawArray(const int32 ChannelIndex = 0) const
+	FORCEINLINE TArrayView<AttributeType> GetRawArray(const int32 ChannelIndex = 0) const
 	{
 		if (ArrayPtr == nullptr)
 		{
-			return TArrayView<Type>();
+			return TArrayView<AttributeType>();
 		}
 
-		Type* Element = ArrayPtr->GetArrayForChannel(ChannelIndex).GetElementBase(0);
-		return TArrayView<Type>(Element, GetNumElements() * Extent);
+		AttributeType* Element = ArrayPtr->GetArrayForChannel(ChannelIndex).GetElementBase(0);
+		return TArrayView<AttributeType>(Element, GetNumElements() * Extent);
 	}
 
 	/** Return whether the reference is valid or not */
 	FORCEINLINE bool IsValid() const { return (ArrayPtr != nullptr); }
 
 	/** Return default value for this attribute type */
-	FORCEINLINE Type GetDefaultValue() const { return ArrayPtr->GetDefaultValue(); }
+	FORCEINLINE AttributeType GetDefaultValue() const { return ArrayPtr->GetDefaultValue(); }
 
 	UE_DEPRECATED(5.0, "Please use GetNumChannels().")
 	FORCEINLINE int32 GetNumIndices() const
@@ -861,39 +1079,94 @@ public:
 	/** Get the flags for this attribute array set */
 	FORCEINLINE EMeshAttributeFlags GetFlags() const { return ArrayPtr->GetFlags(); }
 
+	/** Return the extent of this attribute, i.e. the number of array elements it comprises */
 	FORCEINLINE uint32 GetExtent() const { return Extent; }
 
-	/** Copies the given attribute array and index to this index */
-	void Copy(TMeshAttributesConstRef<ElementIDType, AttributeType, true> Src, const int32 DestIndex = 0, const int32 SrcIndex = 0);
-
-	UE_DEPRECATED(5.0, "Please use SetNumChannels().")
-	FORCEINLINE void SetNumIndices(const int32 NumChannels) const
+	/** Set the element with the given ID and index 0 to the provided value */
+	template <typename T = ElementIDType, typename TEnableIf<TIsDerivedFrom<T, FElementID>::Value, int>::Type = 0>
+	FORCEINLINE void Set(const ElementIDType ElementID, TArrayView<const AttributeType> Value) const
 	{
-		ArrayPtr->ArrayType::SetNumChannels(NumChannels);	// note: override virtual dispatch
+		check(Value.Num() == Extent);
+		AttributeType* Element = ArrayPtr->GetArrayForChannel(0).GetElementBase(ElementID.GetValue());
+		for (uint32 Index = 0; Index < Extent; Index++)
+		{
+			Element[Index] = Value[Index];
+		}
 	}
 
-	/** Sets number of channels this attribute has */
-	FORCEINLINE void SetNumChannels(const int32 NumChannels) const
+	/** Set the element with the given ID and channel to the provided value */
+	template <typename T = ElementIDType, typename TEnableIf<TIsDerivedFrom<T, FElementID>::Value, int>::Type = 0>
+	FORCEINLINE void Set(const ElementIDType ElementID, const int32 Channel, TArrayView<const AttributeType> Value) const
 	{
-		ArrayPtr->ArrayType::SetNumChannels(NumChannels);	// note: override virtual dispatch
+		check(Value.Num() == Extent);
+		AttributeType* Element = ArrayPtr->GetArrayForChannel(Channel).GetElementBase(ElementID.GetValue());
+		for (uint32 Index = 0; Index < Extent; Index++)
+		{
+			Element[Index] = Value[Index];
+		}
 	}
 
-	UE_DEPRECATED(5.0, "Please use InsertChannel().")
-	FORCEINLINE void InsertIndex(const int32 Index) const
+	FORCEINLINE void Set(int32 ElementIndex, TArrayView<const AttributeType> Value) const
 	{
-		ArrayPtr->ArrayType::InsertChannel(Index);		// note: override virtual dispatch
+		check(Value.Num() == Extent);
+		AttributeType* Element = ArrayPtr->GetArrayForChannel(0).GetElementBase(ElementIndex);
+		for (uint32 Index = 0; Index < Extent; Index++)
+		{
+			Element[Index] = Value[Index];
+		}
 	}
 
-	/** Inserts an attribute channel */
-	FORCEINLINE void InsertChannel(const int32 Index) const
+	FORCEINLINE void Set(int32 ElementIndex, const int32 Channel, TArrayView<const AttributeType> Value) const
 	{
-		ArrayPtr->ArrayType::InsertChannel(Index);		// note: override virtual dispatch
+		check(Value.Num() == Extent);
+		AttributeType* Element = ArrayPtr->GetArrayForChannel(Channel).GetElementBase(ElementIndex);
+		for (uint32 Index = 0; Index < Extent; Index++)
+		{
+			Element[Index] = Value[Index];
+		}
 	}
 
-	UE_DEPRECATED(5.0, "Please use RemoveChannel().")
-	FORCEINLINE void RemoveIndex(const int32 Index) const
+	FORCEINLINE void SetArrayView(int32 ElementIndex, TArrayView<const AttributeType> Value) const
 	{
-		ArrayPtr->ArrayType::RemoveChannel(Index);		// note: override virtual dispatch
+		Set(ElementIndex, Value);
+	}
+
+	FORCEINLINE void SetArrayView(int32 ElementIndex, const int32 Channel, TArrayView<const AttributeType> Value) const
+	{
+		Set(ElementIndex, Channel, Value);
+	}
+
+	/** Copies the given attribute array and index to this index */
+	void Copy(TMeshAttributesConstRef<ElementIDType, TArrayView<AttributeType>> Src, const int32 DestChannel = 0, const int32 SrcChannel = 0);
+
+	UE_DEPRECATED(5.0, "Please use SetNumChannels().")
+	FORCEINLINE void SetNumIndices(const int32 NumChannels) const
+	{
+		ArrayPtr->ArrayType::SetNumChannels(NumChannels);	// note: override virtual dispatch
+	}
+
+	/** Sets number of channels this attribute has */
+	FORCEINLINE void SetNumChannels(const int32 NumChannels) const
+	{
+		ArrayPtr->ArrayType::SetNumChannels(NumChannels);	// note: override virtual dispatch
+	}
+
+	UE_DEPRECATED(5.0, "Please use InsertChannel().")
+	FORCEINLINE void InsertIndex(const int32 Index) const
+	{
+		ArrayPtr->ArrayType::InsertChannel(Index);		// note: override virtual dispatch
+	}
+
+	/** Inserts an attribute channel */
+	FORCEINLINE void InsertChannel(const int32 Index) const
+	{
+		ArrayPtr->ArrayType::InsertChannel(Index);		// note: override virtual dispatch
+	}
+
+	UE_DEPRECATED(5.0, "Please use RemoveChannel().")
+	FORCEINLINE void RemoveIndex(const int32 Index) const
+	{
+		ArrayPtr->ArrayType::RemoveChannel(Index);		// note: override virtual dispatch
 	}
 
 	/** Removes an attribute channel */
@@ -908,12 +1181,12 @@ protected:
 };
 
 template <typename ElementIDType, typename AttributeType>
-void TMeshAttributesRef<ElementIDType, AttributeType, true>::Copy(TMeshAttributesConstRef<ElementIDType, AttributeType, true> Src, const int32 DestIndex, const int32 SrcIndex)
+void TMeshAttributesRef<ElementIDType, TArrayView<AttributeType>>::Copy(TMeshAttributesConstRef<ElementIDType, TArrayView<AttributeType>> Src, const int32 DestChannel, const int32 SrcChannel)
 {
 	check(Src.IsValid());
 	check(Src.Extent == Extent);
-	const TMeshAttributeArrayBase<Type>& SrcArray = Src->ArrayPtr->GetArrayForChannel(SrcIndex);
-	TMeshAttributeArrayBase<Type>& DestArray = ArrayPtr->GetArrayForChannel(DestIndex);
+	const TMeshAttributeArrayBase<AttributeType>& SrcArray = Src->ArrayPtr->GetArrayForChannel(SrcChannel);
+	TMeshAttributeArrayBase<AttributeType>& DestArray = ArrayPtr->GetArrayForChannel(DestChannel);
 	const int32 Num = FMath::Min(SrcArray.Num(), DestArray.Num());
 	for (int32 Index = 0; Index < Num; Index++)
 	{
@@ -924,103 +1197,148 @@ void TMeshAttributesRef<ElementIDType, AttributeType, true>::Copy(TMeshAttribute
 	}
 }
 
-/**
- * This is the class used to provide a 'view' of the specified type on an attribute array.
- * Like TMeshAttributesRef, it is a proxy object which is valid for as long as the owning FMeshDescription exists,
- * and should be passed by value.
- *
- * This is the base class, and shouldn't be instanced directly.
- */
-template <typename ViewType>
-class TMeshAttributesViewBase
+
+template <typename ElementIDType, typename AttributeType>
+class TMeshAttributesRef<ElementIDType, TArrayAttribute<AttributeType>>
 {
+	template <typename T, typename U> friend class TMeshAttributesRef;
+
 public:
-	using Type = ViewType;
-	using ArrayType = typename TCopyQualifiersFromTo<ViewType, FMeshAttributeArraySetBase>::Type;
+	using ArrayType = typename TCopyQualifiersFromTo<AttributeType, TMeshUnboundedAttributeArraySet<typename TRemoveCV<AttributeType>::Type>>::Type;
 
-	/** Return whether the reference is valid or not */
-	FORCEINLINE bool IsValid() const { return ( ArrayPtr != nullptr ); }
+	/** Constructor taking a pointer to a TMeshUnboundedAttributeArraySet */
+	FORCEINLINE explicit TMeshAttributesRef(ArrayType* InArrayPtr = nullptr, uint32 InExtent = 0)
+		: ArrayPtr(InArrayPtr)
+	{}
 
-	/** Return number of indices this attribute has */
-	FORCEINLINE int32 GetNumIndices() const { return ArrayPtr->GetNumIndices(); }
+	/** Implicitly construct a TMeshAttributesRef-to-const from a regular one */
+	template <typename T = AttributeType, typename TEnableIf<TIsSame<T, const T>::Value, int>::Type = 0>
+	FORCEINLINE TMeshAttributesRef(TMeshAttributesRef<ElementIDType, TArrayAttribute<typename TRemoveCV<T>::Type>> InRef)
+		: ArrayPtr(InRef.ArrayPtr)
+	{}
 
-	/** Return default value for this attribute type */
-	FORCEINLINE ViewType GetDefaultValue() const;
+	/** Implicitly construct a TMeshAttributesRef from a TMeshAttributesArray **/
+	template <typename T = ElementIDType, typename U = AttributeType, typename TEnableIf<!TIsSame<T, int32>::Value, int>::Type = 0>
+	FORCEINLINE TMeshAttributesRef(TMeshAttributesRef<int32, TArrayAttribute<U>> InRef)
+		: ArrayPtr(InRef.ArrayPtr)
+	{}
 
-	/** Get the number of elements in this attribute array */
-	FORCEINLINE int32 GetNumElements() const
+	/** Implicitly construct a TMeshAttributesRef-to-const from a TMeshAttributesArray */
+	template <typename T = ElementIDType, typename U = AttributeType, typename TEnableIf<TIsSame<U, const U>::Value, int>::Type = 0, typename TEnableIf<!TIsSame<T, int32>::Value, int>::Type = 0>
+	FORCEINLINE TMeshAttributesRef(TMeshAttributesRef<int32, TArrayAttribute<typename TRemoveCV<U>::Type>> InRef)
+		: ArrayPtr(InRef.ArrayPtr)
+	{}
+
+
+	/** Access elements from attribute channel 0 */
+	template <typename T = ElementIDType, typename TEnableIf<TIsDerivedFrom<T, FElementID>::Value, int>::Type = 0>
+	FORCEINLINE TArrayAttribute<AttributeType> operator[](const ElementIDType ElementID) const
 	{
-		return ArrayPtr->GetNumElements();
+		return TArrayAttribute<AttributeType>(ArrayPtr->GetArrayForChannel(0), ElementID.GetValue());
 	}
 
-protected:
-	/** Constructor taking a pointer to a FMeshAttributeArraySetBase */
-	FORCEINLINE explicit TMeshAttributesViewBase( ArrayType* InArrayPtr )
-		: ArrayPtr( InArrayPtr )
-	{}
+	/** Get the element with the given ID and channel */
+	template <typename T = ElementIDType, typename TEnableIf<TIsDerivedFrom<T, FElementID>::Value, int>::Type = 0>
+	FORCEINLINE TArrayAttribute<AttributeType> Get(const ElementIDType ElementID, const int32 Channel = 0) const
+	{
+		return TArrayAttribute<AttributeType>(ArrayPtr->GetArrayForChannel(Channel), ElementID.GetValue());
+	}
 
-	/** Get the element with the given ID from index 0 */
-	FORCEINLINE ViewType GetByIndex( const int32 ElementIndex ) const;
+	FORCEINLINE TArrayAttribute<AttributeType> operator[](int32 ElementIndex) const
+	{
+		return TArrayAttribute<AttributeType>(ArrayPtr->GetArrayForChannel(0), ElementIndex);
+	}
 
-	/** Get the element with the given element and attribute indices */
-	FORCEINLINE ViewType GetByIndex( const int32 ElementIndex, const int32 AttributeIndex ) const;
+	FORCEINLINE TArrayAttribute<AttributeType> Get(int32 ElementIndex, const int32 Channel = 0) const
+	{
+		return TArrayAttribute<AttributeType>(ArrayPtr->GetArrayForChannel(Channel), ElementIndex);
+	}
 
-	/** Set the attribute index 0 element with the given index to the provided value */
-	FORCEINLINE void SetByIndex( const int32 ElementIndex, const ViewType& Value ) const;
+	FORCEINLINE TArrayView<AttributeType> GetArrayView(int32 ElementIndex, const int32 Channel = 0) const
+	{
+		return Get(ElementIndex, Channel).ToArrayView();
+	}
 
-	/** Set the element with the given element and attribute indices to the provided value */
-	FORCEINLINE void SetByIndex( const int32 ElementIndex, const int32 AttributeIndex, const ViewType& Value ) const;
+	/** Return whether the reference is valid or not */
+	FORCEINLINE bool IsValid() const { return (ArrayPtr != nullptr); }
 
-	ArrayType* ArrayPtr;
-};
+	/** Return default value for this attribute type */
+	FORCEINLINE AttributeType GetDefaultValue() const { return ArrayPtr->GetDefaultValue(); }
 
+	/** Return number of channels this attribute has */
+	FORCEINLINE int32 GetNumChannels() const
+	{
+		return ArrayPtr->ArrayType::GetNumChannels();	// note: override virtual dispatch
+	}
 
-/**
- * This is a derived version with typesafe element accessors, which is returned by TAttributesSet<>.
- */
-template <typename ElementIDType, typename ViewType>
-class TMeshAttributesView final : public TMeshAttributesViewBase<ViewType>
-{
-	using Super = TMeshAttributesViewBase<ViewType>;
+	/** Get the number of elements in this attribute array */
+	FORCEINLINE int32 GetNumElements() const
+	{
+		return ArrayPtr->GetNumElements();
+	}
 
-	template <typename T, typename U> friend class TMeshAttributesView;
+	/** Get the flags for this attribute array set */
+	FORCEINLINE EMeshAttributeFlags GetFlags() const { return ArrayPtr->GetFlags(); }
 
-public:
-	/** Constructor taking a pointer to a FMeshAttributeArraySetBase */
-	FORCEINLINE explicit TMeshAttributesView( typename Super::ArrayType* InArrayPtr = nullptr )
-		: Super( InArrayPtr )
-	{}
+	/** Set the element with the given ID and index 0 to the provided value */
+	template <typename T = ElementIDType, typename TEnableIf<TIsDerivedFrom<T, FElementID>::Value, int>::Type = 0>
+	FORCEINLINE void Set(const ElementIDType ElementID, TArrayAttribute<const AttributeType> Value) const
+	{
+		ArrayPtr->GetArrayForChannel(0).Set(ElementID.GetValue(), Value.ToArrayView());
+	}
 
-	/** Implicitly construct a TMeshAttributesViewBase-to-const from a regular one */
-	template <typename T = ViewType, typename TEnableIf<TIsSame<T, const T>::Value, int>::Type = 0>
-	FORCEINLINE TMeshAttributesView( TMeshAttributesView<ElementIDType, typename TRemoveCV<T>::Type> InView )
-		: Super( InView.ArrayPtr )
-	{}
+	/** Set the element with the given ID and channel to the provided value */
+	template <typename T = ElementIDType, typename TEnableIf<TIsDerivedFrom<T, FElementID>::Value, int>::Type = 0>
+	FORCEINLINE void Set(const ElementIDType ElementID, const int32 Channel, TArrayAttribute<const AttributeType> Value) const
+	{
+		ArrayPtr->GetArrayForChannel(Channel).Set(ElementID.GetValue(), Value.ToArrayView());
+	}
 
-	/** Get the element with the given ID from index 0. This version has a typesafe element ID accessor. */
-	FORCEINLINE ViewType Get( const ElementIDType ElementID ) const { return this->GetByIndex( ElementID.GetValue() ); }
+	FORCEINLINE void Set(int32 ElementIndex, TArrayAttribute<const AttributeType> Value) const
+	{
+		ArrayPtr->GetArrayForChannel(0).Set(ElementIndex, Value.ToArrayView());
+	}
 
-	/** Get the element with the given ID and index. This version has a typesafe element ID accessor. */
-	FORCEINLINE ViewType Get( const ElementIDType ElementID, const int32 Index ) const { return this->GetByIndex( ElementID.GetValue(), Index ); }
+	FORCEINLINE void Set(int32 ElementIndex, const int32 Channel, TArrayAttribute<const AttributeType> Value) const
+	{
+		ArrayPtr->GetArrayForChannel(Channel).Set(ElementIndex, Value.ToArrayView());
+	}
 
-	/** Set the element with the given ID and index 0 to the provided value. This version has a typesafe element ID accessor. */
-	FORCEINLINE void Set( const ElementIDType ElementID, const ViewType& Value ) const { this->SetByIndex( ElementID.GetValue(), Value ); }
+	FORCEINLINE void SetArrayView(int32 ElementIndex, TArrayView<const AttributeType> Value) const
+	{
+		ArrayPtr->GetArrayForChannel(0).Set(ElementIndex, Value);
+	}
 
-	/** Set the element with the given ID and index to the provided value. This version has a typesafe element ID accessor. */
-	FORCEINLINE void Set( const ElementIDType ElementID, const int32 Index, const ViewType& Value ) const { this->SetByIndex( ElementID.GetValue(), Index, Value ); }
+	FORCEINLINE void SetArrayView(int32 ElementIndex, const int32 Channel, TArrayView<const AttributeType> Value) const
+	{
+		ArrayPtr->GetArrayForChannel(Channel).Set(ElementIndex, Value);
+	}
 
-	/** Sets number of indices this attribute has */
-	FORCEINLINE void SetNumIndices( const int32 NumIndices ) const { this->ArrayPtr->SetNumChannels( NumIndices ); }
+	/** Copies the given attribute array and index to this index */
+	void Copy(TMeshAttributesConstRef<ElementIDType, TArrayAttribute<AttributeType>> Src, const int32 DestChannel = 0, const int32 SrcChannel = 0);
 
-	/** Inserts an attribute index */
-	FORCEINLINE void InsertIndex( const int32 Index ) const { this->ArrayPtr->InsertIndex( Index ); }
+	/** Sets number of channels this attribute has */
+	FORCEINLINE void SetNumChannels(const int32 NumChannels) const
+	{
+		ArrayPtr->ArrayType::SetNumChannels(NumChannels);	// note: override virtual dispatch
+	}
 
-	/** Removes an attribute index */
-	FORCEINLINE void RemoveIndex( const int32 Index ) const { this->ArrayPtr->RemoveIndex( Index ); }
-};
+	/** Inserts an attribute channel */
+	FORCEINLINE void InsertChannel(const int32 Index) const
+	{
+		ArrayPtr->ArrayType::InsertChannel(Index);		// note: override virtual dispatch
+	}
 
-template <typename ElementIDType, typename AttributeType>
-using TMeshAttributesConstView = TMeshAttributesView<ElementIDType, const AttributeType>;
+	/** Removes an attribute channel */
+	FORCEINLINE void RemoveChannel(const int32 Index) const
+	{
+		ArrayPtr->ArrayType::RemoveChannel(Index);		// note: override virtual dispatch
+	}
+
+protected:
+	ArrayType* ArrayPtr;
+	uint32 Extent;
+};
 
 
 /**
@@ -1035,18 +1353,26 @@ public:
 	 * Default constructor.
 	 * This breaks the invariant that Ptr be always valid, but is necessary so that it can be the value type of a TMap.
 	 */
-	FORCEINLINE FAttributesSetEntry() = default;
+	FAttributesSetEntry() = default;
 
 	/**
 	 * Construct a valid FAttributesSetEntry of the concrete type specified.
 	 */
 	template <typename AttributeType>
-	FORCEINLINE FAttributesSetEntry(const int32 NumberOfChannels, const AttributeType& Default, const EMeshAttributeFlags Flags, const int32 NumElements, const int32 Extent)
-		: Ptr(MakeUnique<TMeshAttributeArraySet<AttributeType>>(NumberOfChannels, Default, Flags, NumElements, Extent))
-	{}
+	FAttributesSetEntry(const int32 NumberOfChannels, const AttributeType& Default, const EMeshAttributeFlags Flags, const int32 NumElements, const int32 Extent)
+	{
+		if (Extent > 0)
+		{
+			Ptr = MakeUnique<TMeshAttributeArraySet<AttributeType>>(NumberOfChannels, Default, Flags, NumElements, Extent);
+		}
+		else
+		{
+			Ptr = MakeUnique<TMeshUnboundedAttributeArraySet<AttributeType>>(NumberOfChannels, Default, Flags, NumElements);
+		}
+	}
 
 	/** Default destructor */
-	FORCEINLINE ~FAttributesSetEntry() = default;
+	~FAttributesSetEntry() = default;
 
 	/** Polymorphic copy: a new copy of Other is created */
 	FAttributesSetEntry(const FAttributesSetEntry& Other)
@@ -1113,35 +1439,87 @@ public:
 	 *		TVertexInstanceAttributeArray<FVector2D>& UV0 = VertexInstanceAttributes().GetAttributes<FVector2D>( "UV", 0 );
 	 *		UV0[ VertexInstanceID ] = FVector2D( 1.0f, 1.0f );
 	 */
-	template <typename AttributeType,
-			  typename TEnableIf<!TIsBoundedArray<AttributeType>::Value, int>::Type = 0>
-	TMeshAttributesArray<AttributeType> RegisterAttribute(const FName AttributeName, const int32 NumberOfChannels = 1, const AttributeType& Default = AttributeType(), const EMeshAttributeFlags Flags = EMeshAttributeFlags::None );
 
-	template <typename ArrayType,
-			  typename TEnableIf<TIsBoundedArray<ArrayType>::Value, int>::Type = 0,
-			  typename AttributeType = typename TBreakArray<ArrayType>::Type,
-			  uint32 Extent = TBreakArray<ArrayType>::Value>
-	TMeshAttributesArray<TArrayView<AttributeType>> RegisterAttribute(const FName AttributeName, const int32 NumberOfChannels = 1, const typename TBreakArray<ArrayType>::Type& Default = AttributeType(), const EMeshAttributeFlags Flags = EMeshAttributeFlags::None);
+	template <typename T>
+	TMeshAttributesArray<typename TMeshAttributesRegisterType<T>::RefType> RegisterAttributeInternal(
+		const FName AttributeName,
+		const int32 NumberOfChannels = 1,
+		const typename TMeshAttributesRegisterType<T>::RealAttributeType& Default = typename TMeshAttributesRegisterType<T>::RealAttributeType(),
+		const EMeshAttributeFlags Flags = EMeshAttributeFlags::None)
+	{
+		using AttributeType = typename TMeshAttributesRegisterType<T>::AttributeType;
+		using RealAttributeType = typename TMeshAttributesRegisterType<T>::RealAttributeType;
+		using RefType = typename TMeshAttributesRegisterType<T>::RefType;
+		using SetType = typename TMeshAttributesRegisterType<T>::SetType;
+		const uint32 Extent = TMeshAttributesRegisterType<T>::Extent;
+
+		if (FAttributesSetEntry* ArraySetPtr = Map.Find(AttributeName))
+		{
+			if ((*ArraySetPtr)->HasType<RealAttributeType>() && (*ArraySetPtr)->GetExtent() == Extent)
+			{
+				static_cast<SetType*>(ArraySetPtr->Get())->SetType::SetNumChannels(NumberOfChannels);	// note: override virtual dispatch
+				(*ArraySetPtr)->SetFlags(Flags);
+				return TMeshAttributesArray<RefType>(static_cast<SetType*>(ArraySetPtr->Get()), Extent);
+			}
+			else
+			{
+				Map.Remove(AttributeName);
+			}
+		}
+
+		FAttributesSetEntry& Entry = Map.Emplace(AttributeName, FAttributesSetEntry(NumberOfChannels, Default, Flags, NumElements, Extent));
+		return TMeshAttributesArray<RefType>(static_cast<SetType*>(Entry.Get()), Extent);
+	}
 
 	/**
-	 * Register a new index attribute (type is implicitly int).
-	 *
-	 * If the attribute name is already registered, it will update it to use the new type, number of channels and flags.
+	 * Register a new simple attribute.
+	 * e.g. RegisterAttribute<float>(...)
+	 * 
+	 * Obtain a reference to this with GetAttributesRef<float>(...)
 	 */
-	template <typename AttributeType,
-			  typename TEnableIf<TIsSame<AttributeType, int>::Value || TIsDerivedFrom<AttributeType, FElementID>::Value, int>::Type = 0>
-	TMeshAttributesArray<AttributeType> RegisterIndexAttribute(const FName AttributeName, const int32 NumberOfChannels = 1, const EMeshAttributeFlags Flags = EMeshAttributeFlags::None);
+	template <typename T,
+			  typename TEnableIf<!TIsArray<T>::Value, int>::Type = 0>
+	TMeshAttributesArray<typename TMeshAttributesRegisterType<T>::RefType> RegisterAttribute(
+		const FName AttributeName,
+		const int32 NumberOfChannels = 1,
+		const T& Default = T(),
+		const EMeshAttributeFlags Flags = EMeshAttributeFlags::None)
+	{
+		return this->RegisterAttributeInternal<T>(AttributeName, NumberOfChannels, Default, Flags);
+	}
 
 	/**
-	 * Register a new attribute denoting an array of indices (type is implicitly int).
-	 *
-	 * If the attribute name is already registered, it will update it to use the new type, number of channels and flags.
+	 * Register a new fixed array attribute.
+	 * e.g. RegisterAttribute<float[3]>(...)
+	 * 
+	 * Obtain a reference to this with GetAttributesRef<TArrayView<float>>(...)
 	 */
-	template <typename ArrayType,
-			  typename AttributeType = typename TBreakArray<ArrayType>::Type,
-			  typename TEnableIf<TIsSame<AttributeType, int>::Value || TIsDerivedFrom<AttributeType, FElementID>::Value, int>::Type = 0,
-			  uint32 Extent = TBreakArray<ArrayType>::Value>
-	TMeshAttributesArray<TArrayView<AttributeType>> RegisterIndexAttribute(const FName AttributeName, const int32 NumberOfChannels = 1, const EMeshAttributeFlags Flags = EMeshAttributeFlags::None);
+	template <typename T,
+			  typename TEnableIf<TIsArray<T>::Value, int>::Type = 0>
+	TMeshAttributesArray<typename TMeshAttributesRegisterType<T>::RefType> RegisterAttribute(
+		const FName AttributeName,
+		const int32 NumberOfChannels = 1,
+		const typename TMeshAttributesRegisterType<T>::RealAttributeType& Default = typename TMeshAttributesRegisterType<T>::RealAttributeType(),
+		const EMeshAttributeFlags Flags = EMeshAttributeFlags::None)
+	{
+		return this->RegisterAttributeInternal<T>(AttributeName, NumberOfChannels, Default, Flags);
+	}
+
+	/**
+	 * Register a new unbounded array attribute.
+	 * e.g. RegisterAttribute<float[]>(...)
+	 * 
+	 * Obtain a reference to this with GetAttributesRef<TArrayAttribute<float>>(...)
+	 */
+	template <typename T,
+			  typename TEnableIf<TIsSame<typename TMeshAttributesRegisterType<T>::RealAttributeType, int>::Value, int>::Type = 0>
+	TMeshAttributesArray<typename TMeshAttributesRegisterType<T>::RefType> RegisterIndexAttribute(
+		const FName AttributeName,
+		const int32 NumberOfChannels = 1,
+		const EMeshAttributeFlags Flags = EMeshAttributeFlags::None)
+	{
+		return this->RegisterAttributeInternal<T>(AttributeName, NumberOfChannels, int32(INDEX_NONE), Flags | EMeshAttributeFlags::IndexReference);
+	}
 
 	/**
 	 * Unregister an attribute with the given name.
@@ -1160,27 +1538,16 @@ public:
 	/**
 	 * Determines whether an attribute of the given type exists with the given name
 	 */
-	template <typename AttributeType,
-			  typename TEnableIf<!TIsBoundedArray<AttributeType>::Value, int>::Type = 0>
+	template <typename T>
 	bool HasAttributeOfType(const FName AttributeName) const
 	{
-		if (const FAttributesSetEntry* ArraySetPtr = Map.Find(AttributeName))
-		{
-			return (*ArraySetPtr)->HasType<AttributeType>() && (*ArraySetPtr)->GetExtent() == 1;
-		}
-
-		return false;
-	}
+		using RealAttributeType = typename TMeshAttributesRefType<T>::RealAttributeType;
 
-	template <typename ArrayType,
-			  typename TEnableIf<TIsBoundedArray<ArrayType>::Value, int>::Type = 0,
-			  typename AttributeType = typename TBreakArray<ArrayType>::Type,
-			  uint32 Extent = TBreakArray<ArrayType>::Value>
-	bool HasAttributeOfType(const FName AttributeName) const
-	{
 		if (const FAttributesSetEntry* ArraySetPtr = Map.Find(AttributeName))
 		{
-			return (*ArraySetPtr)->HasType<AttributeType>() && (*ArraySetPtr)->GetExtent() == Extent;
+			return (*ArraySetPtr)->HasType<RealAttributeType>() &&
+				   (*ArraySetPtr)->GetExtent() >= TMeshAttributesRefType<T>::MinExpectedExtent &&
+				   (*ArraySetPtr)->GetExtent() <= TMeshAttributesRefType<T>::MaxExpectedExtent;
 		}
 
 		return false;
@@ -1284,70 +1651,50 @@ public:
 	 * Get an attribute array with the given type and name.
 	 * The attribute type must correspond to the type passed as the template parameter.
 	 */
-	template <typename AttributeType,
-			  typename TEnableIf<!TIsArrayView<AttributeType>::Value, int>::Type = 0>
-	TMeshAttributesConstRef<int32, AttributeType> GetAttributesRef(const FName AttributeName) const
+	template <typename T>
+	TMeshAttributesConstRef<int32, typename TMeshAttributesRefType<T>::ConstRefType> GetAttributesRef(const FName AttributeName) const
 	{
+		using RefType = typename TMeshAttributesRefType<T>::ConstRefType;
+
 		if (const FAttributesSetEntry* ArraySetPtr = this->Map.Find(AttributeName))
 		{
-			using Type = typename TChooseClass<TIsDerivedFrom<AttributeType, FElementID>::Value, int, AttributeType>::Result;
-			if ((*ArraySetPtr)->HasType<Type>() && (*ArraySetPtr)->GetExtent() == 1)
-			{
-				return TMeshAttributesConstRef<int32, AttributeType>(static_cast<const TMeshAttributeArraySet<AttributeType>*>(ArraySetPtr->Get()));
-			}
-		}
-
-		return TMeshAttributesConstRef<int32, AttributeType>();
-	}
+			using AttributeType = typename TMeshAttributesRefType<T>::AttributeType;
+			using RealAttributeType = typename TMeshAttributesRefType<T>::RealAttributeType;
 
-	template <typename AttributeType,
-			  typename TEnableIf<!TIsArrayView<AttributeType>::Value, int>::Type = 0>
-	TMeshAttributesRef<int32, AttributeType> GetAttributesRef(const FName AttributeName)
-	{
-		if (FAttributesSetEntry* ArraySetPtr = this->Map.Find(AttributeName))
-		{
-			using Type = typename TChooseClass<TIsDerivedFrom<AttributeType, FElementID>::Value, int, AttributeType>::Result;
-			if ((*ArraySetPtr)->HasType<Type>() && (*ArraySetPtr)->GetExtent() == 1)
+			if ((*ArraySetPtr)->HasType<RealAttributeType>())
 			{
-				return TMeshAttributesRef<int32, AttributeType>(static_cast<TMeshAttributeArraySet<AttributeType>*>(ArraySetPtr->Get()));
+				uint32 ActualExtent = (*ArraySetPtr)->GetExtent();
+				if (ActualExtent >= TMeshAttributesRefType<T>::MinExpectedExtent && ActualExtent <= TMeshAttributesRefType<T>::MaxExpectedExtent)
+				{
+					return TMeshAttributesConstRef<int32, RefType>(static_cast<typename TMeshAttributesRefType<T>::ConstSetType*>(ArraySetPtr->Get()), ActualExtent);
+				}
 			}
 		}
 
-		return TMeshAttributesRef<int32, AttributeType>();
+		return TMeshAttributesConstRef<int32, RefType>();
 	}
 
-	template <typename ArrayType,
-			  typename TEnableIf<TIsArrayView<ArrayType>::Value, int>::Type = 0,
-			  typename AttributeType = typename TBreakArrayView<ArrayType>::Type>
-	TMeshAttributesConstRef<int32, ArrayType> GetAttributesRef(const FName AttributeName) const
+	template <typename T>
+	TMeshAttributesRef<int32, typename TMeshAttributesRefType<T>::RefType> GetAttributesRef(const FName AttributeName)
 	{
-		if (const FAttributesSetEntry* ArraySetPtr = this->Map.Find(AttributeName))
-		{
-			using Type = typename TChooseClass<TIsDerivedFrom<AttributeType, FElementID>::Value, int, AttributeType>::Result;
-			if ((*ArraySetPtr)->HasType<Type>())
-			{
-				return TMeshAttributesConstRef<int32, ArrayType>(static_cast<const TMeshAttributeArraySet<AttributeType>*>(ArraySetPtr->Get()), (*ArraySetPtr)->GetExtent());
-			}
-		}
-
-		return TMeshAttributesConstRef<int32, ArrayType>();
-	}
+		using RefType = typename TMeshAttributesRefType<T>::RefType;
 
-	template <typename ArrayType,
-			  typename TEnableIf<TIsArrayView<ArrayType>::Value, int>::Type = 0,
-			  typename AttributeType = typename TBreakArrayView<ArrayType>::Type>
-	TMeshAttributesRef<int32, ArrayType> GetAttributesRef(const FName AttributeName)
-	{
 		if (FAttributesSetEntry* ArraySetPtr = this->Map.Find(AttributeName))
 		{
-			using Type = typename TChooseClass<TIsDerivedFrom<AttributeType, FElementID>::Value, int, AttributeType>::Result;
-			if ((*ArraySetPtr)->HasType<Type>())
+			using AttributeType = typename TMeshAttributesRefType<T>::AttributeType;
+			using RealAttributeType = typename TMeshAttributesRefType<T>::RealAttributeType;
+
+			if ((*ArraySetPtr)->HasType<RealAttributeType>())
 			{
-				return TMeshAttributesRef<int32, ArrayType>(static_cast<TMeshAttributeArraySet<AttributeType>*>(ArraySetPtr->Get()), (*ArraySetPtr)->GetExtent());
+				uint32 ActualExtent = (*ArraySetPtr)->GetExtent();
+				if (ActualExtent >= TMeshAttributesRefType<T>::MinExpectedExtent && ActualExtent <= TMeshAttributesRefType<T>::MaxExpectedExtent)
+				{
+					return TMeshAttributesRef<int32, RefType>(static_cast<typename TMeshAttributesRefType<T>::SetType*>(ArraySetPtr->Get()), ActualExtent);
+				}
 			}
 		}
 
-		return TMeshAttributesRef<int32, ArrayType>();
+		return TMeshAttributesRef<int32, RefType>();
 	}
 
 	void AppendAttributesFrom(const FAttributesSetBase& OtherAttributesSet);
@@ -1367,123 +1714,6 @@ protected:
 };
 
 
-template <typename AttributeType,
-		  typename TEnableIf<!TIsBoundedArray<AttributeType>::Value, int>::Type>
-TMeshAttributesArray<AttributeType> FAttributesSetBase::RegisterAttribute(const FName AttributeName, const int32 NumberOfChannels, const AttributeType& Default, const EMeshAttributeFlags Flags)
-{
-	using ArrayType = TMeshAttributeArraySet<AttributeType>;
-	if (FAttributesSetEntry* ArraySetPtr = Map.Find(AttributeName))
-	{
-		if ((*ArraySetPtr)->HasType<AttributeType>() && (*ArraySetPtr)->GetExtent() == 1)
-		{
-			static_cast<ArrayType*>(ArraySetPtr->Get())->ArrayType::SetNumChannels(NumberOfChannels);	// note: override virtual dispatch
-			(*ArraySetPtr)->SetFlags(Flags);
-			return TMeshAttributesArray<AttributeType>(static_cast<ArrayType*>(ArraySetPtr->Get()));
-		}
-		else
-		{
-			Map.Remove(AttributeName);
-			FAttributesSetEntry& Entry = Map.Emplace(AttributeName, FAttributesSetEntry(NumberOfChannels, Default, Flags, NumElements, 1));
-			return TMeshAttributesArray<AttributeType>(static_cast<ArrayType*>(Entry.Get()));
-		}
-	}
-	else
-	{
-		FAttributesSetEntry& Entry = Map.Emplace(AttributeName, FAttributesSetEntry(NumberOfChannels, Default, Flags, NumElements, 1));
-		return TMeshAttributesArray<AttributeType>(static_cast<ArrayType*>(Entry.Get()));
-	}
-}
-
-
-template <typename ArrayType,
-		  typename TEnableIf<TIsBoundedArray<ArrayType>::Value, int>::Type,
-		  typename AttributeType,
-		  uint32 Extent>
-TMeshAttributesArray<TArrayView<AttributeType>> FAttributesSetBase::RegisterAttribute(const FName AttributeName, const int32 NumberOfChannels, const typename TBreakArray<ArrayType>::Type& Default, const EMeshAttributeFlags Flags)
-{
-	using ContainerType = TMeshAttributeArraySet<AttributeType>;
-	if (FAttributesSetEntry* ArraySetPtr = Map.Find(AttributeName))
-	{
-		if ((*ArraySetPtr)->HasType<AttributeType>() && (*ArraySetPtr)->GetExtent() == Extent)
-		{
-			static_cast<ContainerType*>(ArraySetPtr->Get())->ContainerType::SetNumChannels(NumberOfChannels);	// note: override virtual dispatch
-			(*ArraySetPtr)->SetFlags(Flags);
-			return TMeshAttributesArray<TArrayView<AttributeType>>(static_cast<ContainerType*>(ArraySetPtr->Get()), Extent);
-		}
-		else
-		{
-			Map.Remove(AttributeName);
-			FAttributesSetEntry& Entry = Map.Emplace(AttributeName, FAttributesSetEntry(NumberOfChannels, Default, Flags, NumElements, Extent));
-			return TMeshAttributesArray<TArrayView<AttributeType>>(static_cast<ContainerType*>(Entry.Get()), Extent);
-		}
-	}
-	else
-	{
-		FAttributesSetEntry& Entry = Map.Emplace(AttributeName, FAttributesSetEntry(NumberOfChannels, Default, Flags, NumElements, Extent));
-		return TMeshAttributesArray<TArrayView<AttributeType>>(static_cast<ContainerType*>(Entry.Get()), Extent);
-	}
-}
-
-
-template <typename AttributeType,
-		  typename TEnableIf<TIsSame<AttributeType, int>::Value || TIsDerivedFrom<AttributeType, FElementID>::Value, int>::Type>
-TMeshAttributesArray<AttributeType> FAttributesSetBase::RegisterIndexAttribute(const FName AttributeName, const int32 NumberOfChannels, const EMeshAttributeFlags Flags)
-{
-	using ArrayType = TMeshAttributeArraySet<AttributeType>;
-	if (FAttributesSetEntry* ArraySetPtr = Map.Find(AttributeName))
-	{
-		if ((*ArraySetPtr)->HasType<int>() && (*ArraySetPtr)->GetExtent() == 1)
-		{
-			static_cast<ArrayType*>(ArraySetPtr->Get())->ArrayType::SetNumChannels(NumberOfChannels);	// note: override virtual dispatch
-			(*ArraySetPtr)->SetFlags(Flags);
-			return TMeshAttributesArray<AttributeType>(static_cast<ArrayType*>(ArraySetPtr->Get()));
-		}
-		else
-		{
-			Map.Remove(AttributeName);
-			FAttributesSetEntry& Entry = Map.Emplace(AttributeName, FAttributesSetEntry(NumberOfChannels, int32(INDEX_NONE), Flags | EMeshAttributeFlags::IndexReference, NumElements, 1));
-			return TMeshAttributesArray<AttributeType>(static_cast<ArrayType*>(Entry.Get()));
-		}
-	}
-	else
-	{
-		FAttributesSetEntry& Entry = Map.Emplace(AttributeName, FAttributesSetEntry(NumberOfChannels, int32(INDEX_NONE), Flags | EMeshAttributeFlags::IndexReference, NumElements, 1));
-		return TMeshAttributesArray<AttributeType>(static_cast<ArrayType*>(Entry.Get()));
-	}
-
-}
-
-
-template <typename ArrayType,
-		  typename AttributeType,
-		  typename TEnableIf<TIsSame<AttributeType, int>::Value || TIsDerivedFrom<AttributeType, FElementID>::Value, int>::Type,
-		  uint32 Extent>
-TMeshAttributesArray<TArrayView<AttributeType>> FAttributesSetBase::RegisterIndexAttribute(const FName AttributeName, const int32 NumberOfChannels, const EMeshAttributeFlags Flags)
-{
-	if (FAttributesSetEntry* ArraySetPtr = Map.Find(AttributeName))
-	{
-		if ((*ArraySetPtr)->HasType<int>() && (*ArraySetPtr)->GetExtent() == Extent)
-		{
-			using ContainerType = TMeshAttributeArraySet<int>;
-			static_cast<ContainerType*>(ArraySetPtr->Get())->ContainerType::SetNumChannels(NumberOfChannels);	// note: override virtual dispatch
-			(*ArraySetPtr)->SetFlags(Flags);
-			return TMeshAttributesArray<TArrayView<AttributeType>>(static_cast<TMeshAttributeArraySet<AttributeType>*>(ArraySetPtr->Get()), Extent);
-		}
-		else
-		{
-			Map.Remove(AttributeName);
-			FAttributesSetEntry& Entry = Map.Emplace(AttributeName, FAttributesSetEntry(NumberOfChannels, int32(INDEX_NONE), Flags | EMeshAttributeFlags::IndexReference, NumElements, Extent));
-			return TMeshAttributesArray<TArrayView<AttributeType>>(static_cast<TMeshAttributeArraySet<AttributeType>*>(Entry.Get()), Extent);
-		}
-	}
-	else
-	{
-		FAttributesSetEntry& Entry = Map.Emplace(AttributeName, FAttributesSetEntry(NumberOfChannels, int32(INDEX_NONE), Flags | EMeshAttributeFlags::IndexReference, NumElements, Extent));
-		return TMeshAttributesArray<TArrayView<AttributeType>>(static_cast<TMeshAttributeArraySet<AttributeType>*>(Entry.Get()), Extent);
-	}
-}
-
-
 /**
  * This is a version of the attributes set container which accesses elements by typesafe IDs.
  * This prevents access of (for example) vertex instance attributes by vertex IDs.
@@ -1511,112 +1741,59 @@ public:
 	 * Note that the returned object is a value type which should be assigned and passed by value, not reference.
 	 * It is valid for as long as this TAttributesSet object exists.
 	 */
-	template <typename AttributeType,
-			  typename TEnableIf<!TIsArrayView<AttributeType>::Value, int>::Type = 0>
-	TMeshAttributesConstRef<ElementIDType, AttributeType> GetAttributesRef(const FName AttributeName) const
+	template <typename T>
+	TMeshAttributesConstRef<ElementIDType, typename TMeshAttributesRefType<T>::ConstRefType> GetAttributesRef(const FName AttributeName) const
 	{
+		using RefType = typename TMeshAttributesRefType<T>::ConstRefType;
+
 		if (const FAttributesSetEntry* ArraySetPtr = this->Map.Find(AttributeName))
 		{
-			using Type = typename TChooseClass<TIsDerivedFrom<AttributeType, FElementID>::Value, int32, AttributeType>::Result;
-			if ((*ArraySetPtr)->HasType<Type>() && (*ArraySetPtr)->GetExtent() == 1)
-			{
-				return TMeshAttributesConstRef<ElementIDType, AttributeType>(static_cast<const TMeshAttributeArraySet<AttributeType>*>(ArraySetPtr->Get()));
-			}
-		}
+			using AttributeType = typename TMeshAttributesRefType<T>::AttributeType;
 
-		return TMeshAttributesConstRef<ElementIDType, AttributeType>();
-	}
+			// Get the real attribute type; this will normally be the same as the templated type,
+			// but element ID types are coerced to simple int32s.
+			using RealAttributeType = typename TChooseClass<TIsDerivedFrom<AttributeType, FElementID>::Value, int32, AttributeType>::Result;
 
-	template <typename AttributeType,
-			  typename TEnableIf<!TIsArrayView<AttributeType>::Value, int>::Type = 0>
-	TMeshAttributesRef<ElementIDType, AttributeType> GetAttributesRef(const FName AttributeName)
-	{
-		if (FAttributesSetEntry* ArraySetPtr = this->Map.Find(AttributeName))
-		{
-			using Type = typename TChooseClass<TIsDerivedFrom<AttributeType, FElementID>::Value, int32, AttributeType>::Result;
-			if ((*ArraySetPtr)->HasType<Type>() && (*ArraySetPtr)->GetExtent() == 1)
+			if ((*ArraySetPtr)->HasType<RealAttributeType>())
 			{
-				return TMeshAttributesRef<ElementIDType, AttributeType>(static_cast<TMeshAttributeArraySet<AttributeType>*>(ArraySetPtr->Get()));
+				uint32 ActualExtent = (*ArraySetPtr)->GetExtent();
+				if (ActualExtent >= TMeshAttributesRefType<T>::MinExpectedExtent && ActualExtent <= TMeshAttributesRefType<T>::MaxExpectedExtent)
+				{
+					return TMeshAttributesConstRef<ElementIDType, RefType>(static_cast<typename TMeshAttributesRefType<T>::ConstSetType*>(ArraySetPtr->Get()), ActualExtent);
+				}
 			}
 		}
 
-		return TMeshAttributesRef<ElementIDType, AttributeType>();
+		return TMeshAttributesConstRef<ElementIDType, RefType>();
 	}
 
-	template <typename ArrayType,
-			  typename TEnableIf<TIsArrayView<ArrayType>::Value, int>::Type = 0,
-			  typename AttributeType = typename TBreakArrayView<ArrayType>::Type>
-	TMeshAttributesConstRef<ElementIDType, ArrayType> GetAttributesRef(const FName AttributeName) const
+	// Non-const version
+	template <typename T>
+	TMeshAttributesRef<ElementIDType, typename TMeshAttributesRefType<T>::RefType> GetAttributesRef(const FName AttributeName)
 	{
-		if (const FAttributesSetEntry* ArraySetPtr = this->Map.Find(AttributeName))
-		{
-			using Type = typename TChooseClass<TIsDerivedFrom<AttributeType, FElementID>::Value, int32, AttributeType>::Result;
-			if ((*ArraySetPtr)->HasType<Type>())
-			{
-				return TMeshAttributesConstRef<ElementIDType, ArrayType>(static_cast<const TMeshAttributeArraySet<AttributeType>*>(ArraySetPtr->Get()), (*ArraySetPtr)->GetExtent());
-			}
-		}
-
-		return TMeshAttributesConstRef<ElementIDType, ArrayType>();
-	}
+		using RefType = typename TMeshAttributesRefType<T>::RefType;
 
-	template <typename ArrayType,
-			  typename TEnableIf<TIsArrayView<ArrayType>::Value, int>::Type = 0,
-			  typename AttributeType = typename TBreakArrayView<ArrayType>::Type>
-	TMeshAttributesRef<ElementIDType, ArrayType> GetAttributesRef(const FName AttributeName)
-	{
 		if (FAttributesSetEntry* ArraySetPtr = this->Map.Find(AttributeName))
 		{
-			using Type = typename TChooseClass<TIsDerivedFrom<AttributeType, FElementID>::Value, int32, AttributeType>::Result;
-			if ((*ArraySetPtr)->HasType<Type>())
-			{
-				return TMeshAttributesRef<ElementIDType, ArrayType>(static_cast<TMeshAttributeArraySet<AttributeType>*>(ArraySetPtr->Get()), (*ArraySetPtr)->GetExtent());
-			}
-		}
+			using AttributeType = typename TMeshAttributesRefType<T>::AttributeType;
 
-		return TMeshAttributesRef<ElementIDType, ArrayType>();
-	}
+			// Get the real attribute type; this will normally be the same as the templated type,
+			// but element ID types are coerced to simple int32s.
+			using RealAttributeType = typename TChooseClass<TIsDerivedFrom<AttributeType, FElementID>::Value, int32, AttributeType>::Result;
 
-	/**
-	 * Get a view on an attribute array with the given name, accessing elements as the given type.
-	 * Access to elements will be slightly slower than with GetAttributesRef, but element access is not strongly typed.
-	 *
-	 * Example of use:
-	 *
-	 *		const TVertexInstanceAttributesView<FVector> VertexNormals = VertexInstanceAttributes().GetAttributesView<FVector>( "Normal" );
-	 *		for( const FVertexInstanceID VertexInstanceID : GetVertexInstances().GetElementIDs() )
-	 *		{
-	 *          // This will work even if the Normals array has a different internal type, e.g. FPackedVector
-	 *			const FVector Normal = VertexNormals.Get( VertexInstanceID );
-	 *			DoSomethingWith( Normal );
-	 *		}
-	 *
-	 * Note that the returned object is a value type which should be assigned and passed by value, not reference.
-	 * It is valid for as long as this TAttributesSet object exists.
-	 */
-	template <typename ViewType>
-	UE_DEPRECATED(4.26, "GetAttributesView() will no longer be supported; please use GetAttributesRef() instead.")
-	TMeshAttributesConstView<ElementIDType, ViewType> GetAttributesView(const FName AttributeName) const
-	{
-		if (const FAttributesSetEntry* ArraySetPtr = this->Map.Find(AttributeName))
-		{
-			return TMeshAttributesConstView<ElementIDType, ViewType>(ArraySetPtr->Get());
+			if ((*ArraySetPtr)->HasType<RealAttributeType>())
+			{
+				uint32 ActualExtent = (*ArraySetPtr)->GetExtent();
+				if (ActualExtent >= TMeshAttributesRefType<T>::MinExpectedExtent && ActualExtent <= TMeshAttributesRefType<T>::MaxExpectedExtent)
+				{
+					return TMeshAttributesRef<ElementIDType, RefType>(static_cast<typename TMeshAttributesRefType<T>::SetType*>(ArraySetPtr->Get()), ActualExtent);
+				}
+			}
 		}
 
-		return TMeshAttributesConstView<ElementIDType, ViewType>();
+		return TMeshAttributesRef<ElementIDType, RefType>();
 	}
 
-	template <typename ViewType>
-	UE_DEPRECATED(4.26, "GetAttributesView() will no longer be supported; please use GetAttributesRef() instead.")
-	TMeshAttributesView<ElementIDType, ViewType> GetAttributesView(const FName AttributeName)
-	{
-		if (FAttributesSetEntry* ArraySetPtr = this->Map.Find(AttributeName))
-		{
-			return TMeshAttributesView<ElementIDType, ViewType>(ArraySetPtr->Get());
-		}
-
-		return TMeshAttributesView<ElementIDType, ViewType>();
-	}
 
 	UE_DEPRECATED(5.0, "Please use GetAttributeChannelCount() instead.")
 	int32 GetAttributeIndexCount(const FName AttributeName) const
@@ -1681,7 +1858,7 @@ public:
 	}
 
 	UE_DEPRECATED(5.0, "Please use InsertAttributeChannel() instead.")
-	void InsertAttributeIndex_Old(const FName AttributeName, const int32 Index)
+	void InsertAttributeIndex(const FName AttributeName, const int32 Index)
 	{
 		InsertAttributeChannel(AttributeName, Index);
 	}
@@ -1710,7 +1887,7 @@ public:
 	}
 
 	UE_DEPRECATED(5.0, "Please use RemoveAttributeChannel() instead.")
-	void RemoveAttributeIndex_Old(const FName AttributeName, const int32 Index)
+	void RemoveAttributeIndex(const FName AttributeName, const int32 Index)
 	{
 		RemoveAttributeChannel(AttributeName, Index);
 	}
@@ -1742,27 +1919,21 @@ public:
 	 * Get an attribute value for the given element ID.
 	 * Note: it is generally preferable to get a TMeshAttributesRef and access elements through that, if you wish to access more than one.
 	 */
-	template <typename AttributeType,
-			  typename TEnableIf<!TIsArrayView<AttributeType>::Value, int>::Type = 0>
-	AttributeType GetAttribute(const ElementIDType ElementID, const FName AttributeName, const int32 AttributeChannel = 0) const
+	template <typename T>
+	T GetAttribute(const ElementIDType ElementID, const FName AttributeName, const int32 AttributeChannel = 0) const
 	{
-		const FMeshAttributeArraySetBase* ArraySetPtr = this->Map.FindChecked(AttributeName).Get();
-		check(ArraySetPtr->HasType<AttributeType>());
-		check(ArraySetPtr->GetExtent() == 1);
-		const AttributeType* ElementBase = static_cast<const TMeshAttributeArraySet<AttributeType>*>(ArraySetPtr)->GetArrayForChannel(AttributeChannel).GetElementBase(ElementID.GetValue());
-		return ElementBase[0];
-	}
+		using RefType = typename TMeshAttributesRefType<T>::ConstRefType;
+		using SetType = typename TMeshAttributesRefType<T>::ConstSetType;
+		using AttributeType = typename TMeshAttributesRefType<T>::AttributeType;
+		using RealAttributeType = typename TMeshAttributesRefType<T>::RealAttributeType;
 
-	/** Get an attribute value for the given element ID. This is a compound (array) attribute, and returns an array view */
-	template <typename ArrayType,
-			  typename TEnableIf<TIsArrayView<ArrayType>::Value, int>::Type = 0,
-			  typename AttributeType = typename TBreakArrayView<ArrayType>::Type>
-	TArrayView<AttributeType> GetAttribute(const ElementIDType ElementID, const FName AttributeName, const int32 AttributeChannel = 0) const
-	{
 		const FMeshAttributeArraySetBase* ArraySetPtr = this->Map.FindChecked(AttributeName).Get();
-		check(ArraySetPtr->HasType<AttributeType>());
-		const AttributeType* ElementBase = static_cast<const TMeshAttributeArraySet<AttributeType>*>(ArraySetPtr)->GetArrayForChannel(AttributeChannel).GetElementBase(ElementID.GetValue());
-		return TArrayView<AttributeType>(ElementBase[0], ArraySetPtr->GetExtent());
+		uint32 ActualExtent = ArraySetPtr->GetExtent();
+		check(ArraySetPtr->HasType<RealAttributeType>());
+		check(ActualExtent >= TMeshAttributesRefType<AttributeType>::MinExpectedExtent && ActualExtent <= TMeshAttributesRefType<AttributeType>::MaxExpectedExtent);
+
+		TMeshAttributesConstRef<FElementID, RefType> Ref(static_cast<SetType*>(ArraySetPtr), ActualExtent);
+		return Ref.Get(ElementID, AttributeChannel);
 	}
 
 
@@ -1770,32 +1941,25 @@ public:
 	 * Set an attribute value for the given element ID.
 	 * Note: it is generally preferable to get a TMeshAttributesRef and set multiple elements through that.
 	 */
-	template <typename AttributeType,
-			  typename TEnableIf<!TIsArrayView<AttributeType>::Value, int>::Type = 0>
-	void SetAttribute(const ElementIDType ElementID, const FName AttributeName, const int32 AttributeChannel, const AttributeType& AttributeValue)
+	template <typename T>
+	void SetAttribute(const ElementIDType ElementID, const FName AttributeName, const int32 AttributeChannel, const T& AttributeValue)
 	{
-		FMeshAttributeArraySetBase* ArraySetPtr = this->Map.FindChecked(AttributeName).Get();
-		check(ArraySetPtr->HasType<AttributeType>());
-		check(ArraySetPtr->GetExtent() == 1);
-		AttributeType* ElementBase = static_cast<TMeshAttributeArraySet<AttributeType>*>(ArraySetPtr)->GetArrayForChannel(AttributeChannel).GetElementBase(ElementID.GetValue());
-		ElementBase[0] = AttributeValue;
-	}
+		using NonConstRefType = typename TMeshAttributesRefType<T>::NonConstRefType;
+		using SetType = typename TMeshAttributesRefType<T>::SetType;
+		using AttributeType = typename TMeshAttributesRefType<T>::AttributeType;
+		using RealAttributeType = typename TMeshAttributesRefType<T>::RealAttributeType;
 
-	template <typename ArrayType,
-			  typename TEnableIf<TIsArrayView<ArrayType>::Value, int>::Type = 0,
-			  typename AttributeType = typename TRemoveCV<typename TBreakArrayView<ArrayType>::Type>::Type>
-	void SetAttribute(const ElementIDType ElementID, const FName AttributeName, const int32 AttributeChannel, const ArrayType& AttributeValue)
-	{
 		FMeshAttributeArraySetBase* ArraySetPtr = this->Map.FindChecked(AttributeName).Get();
-		check(ArraySetPtr->HasType<AttributeType>());
-		check(ArraySetPtr->GetExtent() == AttributeValue.Num());
-		AttributeType* ElementBase = static_cast<TMeshAttributeArraySet<AttributeType>*>(ArraySetPtr)->GetArrayForChannel(AttributeChannel).GetElementBase(ElementID.GetValue());
-		for (int32 I = 0; I < AttributeValue.Num(); I++)
-		{
-			ElementBase[I] = AttributeValue[I];
-		}
+		uint32 ActualExtent = ArraySetPtr->GetExtent();
+		check(ArraySetPtr->HasType<typename TRemoveCV<RealAttributeType>::Type>());
+		check(ActualExtent >= TMeshAttributesRefType<AttributeType>::MinExpectedExtent && ActualExtent <= TMeshAttributesRefType<AttributeType>::MaxExpectedExtent);
+
+		TMeshAttributesRef<FElementID, NonConstRefType> Ref(static_cast<SetType*>(ArraySetPtr), ActualExtent);
+		return Ref.Set(ElementID, AttributeChannel, AttributeValue);
 	}
 
+
+
 	/** Inserts a default-initialized value for all attributes of the given ID */
 	FORCEINLINE void Insert(const ElementIDType ElementID)
 	{
@@ -1854,6 +2018,23 @@ public:
  * The function Dispatch(...) is the actual function which gets called.
  * MakeJumpTable() is the constexpr function which creates a static jump table at compile time.
  */
+
+
+/**
+ * Class which implements a function jump table to be automatically generated at compile time.
+ * This is used by TAttributesSet to provide O(1) dispatch by attribute type at runtime.
+ */
+template <typename FnType, uint32 Size>
+struct TJumpTable
+{
+	template <typename... T>
+	explicit constexpr TJumpTable( T... Ts ) : Fns{ Ts... } {}
+
+	FnType* Fns[Size];
+};
+
+
+
 namespace ForEachImpl
 {
 	// Declare type of jump table used to dispatch functions
@@ -1865,14 +2046,17 @@ namespace ForEachImpl
 	static void Dispatch(FName Name, ForEachFunc Fn, FMeshAttributeArraySetBase* Attributes)
 	{
 		using AttributeType = typename TTupleElement<I, AttributeTypes>::Type;
-		if (Attributes->GetExtent() == 1)
+		if (Attributes->GetExtent() == 0)
+		{
+			Fn(Name, TMeshAttributesRef<ElementIDType, TArrayAttribute<AttributeType>>(static_cast<TMeshUnboundedAttributeArraySet<AttributeType>*>(Attributes)));
+		}
+		else if (Attributes->GetExtent() == 1)
 		{
 			Fn(Name, TMeshAttributesRef<ElementIDType, AttributeType>(static_cast<TMeshAttributeArraySet<AttributeType>*>(Attributes)));
 		}
 		else
 		{
-			// @todo: allow ForEach to iterate through array types
-//			Fn(Name, TMeshAttributesRef<ElementIDType, TArrayView<AttributeType>(static_cast<TMeshAttributeArraySet<AttributeType>*>(Attributes), Attributes->GetExtent()));
+			Fn(Name, TMeshAttributesRef<ElementIDType, TArrayView<AttributeType>>(static_cast<TMeshAttributeArraySet<AttributeType>*>(Attributes), Attributes->GetExtent()));
 		}
 	}
 
@@ -1910,14 +2094,17 @@ namespace ForEachConstImpl
 	static void Dispatch(FName Name, ForEachFunc Fn, const FMeshAttributeArraySetBase* Attributes)
 	{
 		using AttributeType = typename TTupleElement<I, AttributeTypes>::Type;
-		if (Attributes->GetExtent() == 1)
+		if (Attributes->GetExtent() == 0)
+		{
+			Fn(Name, TMeshAttributesConstRef<ElementIDType, TArrayAttribute<const AttributeType>>(static_cast<const TMeshUnboundedAttributeArraySet<AttributeType>*>(Attributes)));
+		}
+		else if (Attributes->GetExtent() == 1)
 		{
 			Fn(Name, TMeshAttributesConstRef<ElementIDType, AttributeType>(static_cast<const TMeshAttributeArraySet<AttributeType>*>(Attributes)));
 		}
 		else
 		{
-			// @todo: allow ForEach to iterate through array types
-//			Fn(Name, TMeshAttributesConstRef<ElementIDType, TArrayView<const AttributeType>>(static_cast<const TMeshAttributeArraySet<AttributeType>*>(Attributes), Attributes->GetExtent()));
+			Fn(Name, TMeshAttributesConstRef<ElementIDType, TArrayView<const AttributeType>>(static_cast<const TMeshAttributeArraySet<AttributeType>*>(Attributes), Attributes->GetExtent()));
 		}
 	}
 
@@ -1960,7 +2147,14 @@ namespace CreateTypeImpl
 	static TUniquePtr<FMeshAttributeArraySetBase> Dispatch(uint32 Extent)
 	{
 		using AttributeType = typename TTupleElement<I, AttributeTypes>::Type;
-		return MakeUnique<TMeshAttributeArraySet<AttributeType>>(Extent);
+		if (Extent > 0)
+		{
+			return MakeUnique<TMeshAttributeArraySet<AttributeType>>(Extent);
+		}
+		else
+		{
+			return MakeUnique<TMeshUnboundedAttributeArraySet<AttributeType>>();
+		}
 	}
 
 	// Build RegisterAttributeOfType jump table at compile time, a separate instantiation of Dispatch for each attribute type
@@ -1977,228 +2171,3 @@ inline void FAttributesSetEntry::CreateArrayOfType(const uint32 Type, const uint
 	Ptr = JumpTable.Fns[Type](Extent);
 }
 
-
-/**
- * Helper struct which determines whether ViewType and the I'th type from AttributeTypes are mutually constructible from each other.
- */
-template <typename ViewType, uint32 I>
-struct TIsViewable
-{
-	enum { Value = TIsConstructible<ViewType, typename TTupleElement<I, AttributeTypes>::Type>::Value &&
-				   TIsConstructible<typename TTupleElement<I, AttributeTypes>::Type, ViewType>::Value };
-};
-
-/**
- * Implementation for TMeshAttributesConstViewBase::Get(ElementIndex).
- *
- * This is implemented similarly to the above. A jump table is built, so the correct implementation is dispatched according to the array type.
- * This cannot be a regular virtual function because the return type depends on the array type.
- */
-namespace AttributesViewGetImpl
-{
-	// Declare type of jump table used to dispatch functions
-	template <typename ViewType>
-	using JumpTableType = TJumpTable<ViewType( const FMeshAttributeArraySetBase*, int32 ), TTupleArity<AttributeTypes>::Value>;
-
-	// Define dispatch functions
-	template <typename ViewType, uint32 I, typename TEnableIf<TIsViewable<ViewType, I>::Value, int>::Type = 0>
-	static ViewType Dispatch( const FMeshAttributeArraySetBase* Array, const int32 Index )
-	{
-		// Implementation when the attribute type is convertible to the view type
-		return ViewType( static_cast<const TMeshAttributeArraySet<typename TTupleElement<I, AttributeTypes>::Type>*>( Array )->GetArrayForChannel( 0 )[ Index ] );
-	}
-
-	template <typename ViewType, uint32 I, typename TEnableIf<!TIsViewable<ViewType, I>::Value, int>::Type = 0>
-	static ViewType Dispatch( const FMeshAttributeArraySetBase* Array, const int32 Index )
-	{
-		// Implementation when the attribute type is not convertible to the view type
-		check( false );
-		return ViewType();
-	}
-
-	// Build jump table at compile time, a separate instantiation of Dispatch for each view type
-	template <typename ViewType, uint32... Is>
-	static constexpr JumpTableType<ViewType> MakeJumpTable( TIntegerSequence< uint32, Is...> )
-	{
-		return JumpTableType<ViewType>( Dispatch<ViewType, Is>... );
-	}
-}
-
-template <typename ViewType>
-FORCEINLINE ViewType TMeshAttributesViewBase<ViewType>::GetByIndex( const int32 ElementIndex ) const
-{
-	static constexpr AttributesViewGetImpl::JumpTableType<ViewType>
-		JumpTable = AttributesViewGetImpl::MakeJumpTable<ViewType>( TMakeIntegerSequence<uint32, TTupleArity<AttributeTypes>::Value>() );
-
-	return JumpTable.Fns[ ArrayPtr->GetType() ]( ArrayPtr, ElementIndex );
-}
-
-
-/**
- * Implementation for TMeshAttributesConstViewBase::Get(ElementIndex, AttributeIndex).
- */
-namespace AttributesViewGetWithIndexImpl
-{
-	// Declare type of jump table used to dispatch functions
-	template <typename ViewType>
-	using JumpTableType = TJumpTable<ViewType( const FMeshAttributeArraySetBase*, int32, int32 ), TTupleArity<AttributeTypes>::Value>;
-
-	// Define dispatch functions
-	template <typename ViewType, uint32 I, typename TEnableIf<TIsViewable<ViewType, I>::Value, int>::Type = 0>
-	static ViewType Dispatch( const FMeshAttributeArraySetBase* Array, const int32 ElementIndex, const int32 AttributeIndex )
-	{
-		return ViewType( static_cast<const TMeshAttributeArraySet<typename TTupleElement<I, AttributeTypes>::Type>*>( Array )->GetArrayForChannel( AttributeIndex )[ ElementIndex ] );
-	}
-
-	template <typename ViewType, uint32 I, typename TEnableIf<!TIsViewable<ViewType, I>::Value, int>::Type = 0>
-	static ViewType Dispatch( const FMeshAttributeArraySetBase* Array, const int32 ElementIndex, const int32 AttributeIndex )
-	{
-		check( false );
-		return ViewType();
-	}
-
-	// Build jump table at compile time, a separate instantiation of Dispatch for each attribute type
-	template <typename ViewType, uint32... Is>
-	static constexpr JumpTableType<ViewType> MakeJumpTable( TIntegerSequence< uint32, Is...> )
-	{
-		return JumpTableType<ViewType>( Dispatch<ViewType, Is>... );
-	}
-}
-
-template <typename ViewType>
-FORCEINLINE ViewType TMeshAttributesViewBase<ViewType>::GetByIndex( const int32 ElementIndex, const int32 AttributeIndex ) const
-{
-	static constexpr AttributesViewGetWithIndexImpl::JumpTableType<ViewType>
-		JumpTable = AttributesViewGetWithIndexImpl::MakeJumpTable<ViewType>( TMakeIntegerSequence<uint32, TTupleArity<AttributeTypes>::Value>() );
-
-	return JumpTable.Fns[ ArrayPtr->GetType() ]( ArrayPtr, ElementIndex, AttributeIndex );
-}
-
-
-/**
- * Implementation for TMeshAttributesViewBase::Set(ElementIndex, Value).
- */
-namespace AttributesViewSetImpl
-{
-	// Declare type of jump table used to dispatch functions
-	template <typename ViewType>
-	using JumpTableType = TJumpTable<void( FMeshAttributeArraySetBase*, int32, const ViewType& ), TTupleArity<AttributeTypes>::Value>;
-
-	// Define dispatch functions
-	template <typename ViewType, uint32 I, typename TEnableIf<TIsViewable<ViewType, I>::Value, int>::Type = 0>
-	static void Dispatch( FMeshAttributeArraySetBase* Array, const int32 Index, const ViewType& Value )
-	{
-		// Implementation when the attribute type is convertible to the view type
-		using AttributeType = typename TTupleElement<I, AttributeTypes>::Type;
-		static_cast<TMeshAttributeArraySet<AttributeType>*>( Array )->GetArrayForChannel( 0 )[ Index ] = AttributeType( Value );
-	}
-
-	template <typename ViewType, uint32 I, typename TEnableIf<!TIsViewable<ViewType, I>::Value, int>::Type = 0>
-	static void Dispatch( FMeshAttributeArraySetBase* Array, const int32 Index, const ViewType& Value )
-	{
-		// Implementation when the attribute type is not convertible to the view type
-		check( false );
-	}
-
-	// Build jump table at compile time, a separate instantiation of Dispatch for each view type
-	template <typename ViewType, uint32... Is>
-	static constexpr JumpTableType<ViewType> MakeJumpTable( TIntegerSequence< uint32, Is...> )
-	{
-		return JumpTableType<ViewType>( Dispatch<ViewType, Is>... );
-	}
-}
-
-template <typename ViewType>
-FORCEINLINE void TMeshAttributesViewBase<ViewType>::SetByIndex( const int32 ElementIndex, const ViewType& Value ) const
-{
-	static constexpr AttributesViewSetImpl::JumpTableType<ViewType>
-		JumpTable = AttributesViewSetImpl::MakeJumpTable<ViewType>( TMakeIntegerSequence<uint32, TTupleArity<AttributeTypes>::Value>() );
-
-	JumpTable.Fns[ ArrayPtr->GetType() ]( ArrayPtr, ElementIndex, Value );
-}
-
-
-/**
- * Implementation for TMeshAttributesViewBase::Set(ElementIndex, AttributeIndex, Value).
- */
-namespace AttributesViewSetWithIndexImpl
-{
-	// Declare type of jump table used to dispatch functions
-	template <typename ViewType>
-	using JumpTableType = TJumpTable<void( FMeshAttributeArraySetBase*, int32, int32, const ViewType& ), TTupleArity<AttributeTypes>::Value>;
-
-	// Define dispatch functions
-	template <typename ViewType, uint32 I, typename TEnableIf<TIsViewable<ViewType, I>::Value, int>::Type = 0>
-	static void Dispatch( FMeshAttributeArraySetBase* Array, const int32 ElementIndex, const int32 AttributeIndex, const ViewType& Value )
-	{
-		// Implementation when the attribute type is convertible to the view type
-		using AttributeType = typename TTupleElement<I, AttributeTypes>::Type;
-		static_cast<TMeshAttributeArraySet<AttributeType>*>( Array )->GetArrayForChannel( AttributeIndex )[ ElementIndex ] = AttributeType( Value );
-	}
-
-	template <typename ViewType, uint32 I, typename TEnableIf<!TIsViewable<ViewType, I>::Value, int>::Type = 0>
-	static void Dispatch( FMeshAttributeArraySetBase* Array, const int32 ElementIndex, const int32 AttributeIndex, const ViewType& Value )
-	{
-		// Implementation when the attribute type is not convertible to the view type
-		check( false );
-	}
-
-	// Build jump table at compile time, a separate instantiation of Dispatch for each view type
-	template <typename ViewType, uint32... Is>
-	static constexpr JumpTableType<ViewType> MakeJumpTable( TIntegerSequence< uint32, Is...> )
-	{
-		return JumpTableType<ViewType>( Dispatch<ViewType, Is>... );
-	}
-}
-
-template <typename ViewType>
-FORCEINLINE void TMeshAttributesViewBase<ViewType>::SetByIndex( const int32 ElementIndex, const int32 AttributeIndex, const ViewType& Value ) const
-{
-	static constexpr AttributesViewSetWithIndexImpl::JumpTableType<ViewType>
-		JumpTable = AttributesViewSetWithIndexImpl::MakeJumpTable<ViewType>( TMakeIntegerSequence<uint32, TTupleArity<AttributeTypes>::Value>() );
-
-	JumpTable.Fns[ ArrayPtr->GetType() ]( ArrayPtr, ElementIndex, AttributeIndex, Value );
-}
-
-
-/**
- * Implementation for TMeshAttributesViewBase::GetDefaultValue().
- */
-namespace AttributesViewGetDefaultImpl
-{
-	// Declare type of jump table used to dispatch functions
-	template <typename ViewType>
-	using JumpTableType = TJumpTable<ViewType( const FMeshAttributeArraySetBase* ), TTupleArity<AttributeTypes>::Value>;
-
-	// Define dispatch functions
-	template <typename ViewType, uint32 I, typename TEnableIf<TIsViewable<ViewType, I>::Value, int>::Type = 0>
-	static ViewType Dispatch( const FMeshAttributeArraySetBase* Array )
-	{
-		// Implementation when the attribute type is convertible to the view type
-		return ViewType( static_cast<const TMeshAttributeArraySet<typename TTupleElement<I, AttributeTypes>::Type>*>( Array )->GetDefaultValue() );
-	}
-
-	template <typename ViewType, uint32 I, typename TEnableIf<!TIsViewable<ViewType, I>::Value, int>::Type = 0>
-	static ViewType Dispatch( const FMeshAttributeArraySetBase* Array )
-	{
-		// Implementation when the attribute type is not convertible to the view type
-		check( false );
-		return ViewType();
-	}
-
-	// Build jump table at compile time, a separate instantiation of Dispatch for each view type
-	template <typename ViewType, uint32... Is>
-	static constexpr JumpTableType<ViewType> MakeJumpTable( TIntegerSequence< uint32, Is...> )
-	{
-		return JumpTableType<ViewType>( Dispatch<ViewType, Is>... );
-	}
-}
-
-template <typename ViewType>
-FORCEINLINE ViewType TMeshAttributesViewBase<ViewType>::GetDefaultValue() const
-{
-	static constexpr AttributesViewGetDefaultImpl::JumpTableType<ViewType>
-		JumpTable = AttributesViewGetDefaultImpl::MakeJumpTable<ViewType>( TMakeIntegerSequence<uint32, TTupleArity<AttributeTypes>::Value>() );
-
-	return JumpTable.Fns[ ArrayPtr->GetType() ]( ArrayPtr );
-}
diff --git a/Engine/Source/Runtime/MeshDescription/Public/MeshDescription.h b/Engine/Source/Runtime/MeshDescription/Public/MeshDescription.h
--- a/Engine/Source/Runtime/MeshDescription/Public/MeshDescription.h
+++ b/Engine/Source/Runtime/MeshDescription/Public/MeshDescription.h
@@ -55,20 +55,6 @@ template <typename AttributeType> using TTriangleAttributesConstRef = TMeshAttri
 template <typename AttributeType> using TPolygonAttributesConstRef = TMeshAttributesConstRef<FPolygonID, AttributeType>;
 template <typename AttributeType> using TPolygonGroupAttributesConstRef = TMeshAttributesConstRef<FPolygonGroupID, AttributeType>;
 
-template <typename AttributeType> using TVertexAttributesView = TMeshAttributesView<FVertexID, AttributeType>;
-template <typename AttributeType> using TVertexInstanceAttributesView = TMeshAttributesView<FVertexInstanceID, AttributeType>;
-template <typename AttributeType> using TEdgeAttributesView = TMeshAttributesView<FEdgeID, AttributeType>;
-template <typename AttributeType> using TTriangleAttributesView = TMeshAttributesView<FTriangleID, AttributeType>;
-template <typename AttributeType> using TPolygonAttributesView = TMeshAttributesView<FPolygonID, AttributeType>;
-template <typename AttributeType> using TPolygonGroupAttributesView = TMeshAttributesView<FPolygonGroupID, AttributeType>;
-
-template <typename AttributeType> using TVertexAttributesConstView = TMeshAttributesConstView<FVertexID, AttributeType>;
-template <typename AttributeType> using TVertexInstanceAttributesConstView = TMeshAttributesConstView<FVertexInstanceID, AttributeType>;
-template <typename AttributeType> using TEdgeAttributesConstView = TMeshAttributesConstView<FEdgeID, AttributeType>;
-template <typename AttributeType> using TTriangleAttributesConstView = TMeshAttributesConstView<FTriangleID, AttributeType>;
-template <typename AttributeType> using TPolygonAttributesConstView = TMeshAttributesConstView<FPolygonID, AttributeType>;
-template <typename AttributeType> using TPolygonGroupAttributesConstView = TMeshAttributesConstView<FPolygonGroupID, AttributeType>;
-
 UENUM()
 enum class EComputeNTBsOptions : uint32
 {
@@ -1191,6 +1177,9 @@ public:
 	/** Remaps the element IDs in the mesh description according to the passed in object */
 	void Remap( const FElementIDRemappings& Remappings );
 
+	/** Gets the number of UV element channels */
+	int32 GetNumUVElementChannels() const { return UVElements->GetNumChannels(); }
+
 	/** Sets the specified number of UV channels */
 	void SetNumUVChannels(const int32 NumUVChannels);
 
diff --git a/Engine/Source/Runtime/StaticMeshDescription/Private/StaticMeshOperations.cpp b/Engine/Source/Runtime/StaticMeshDescription/Private/StaticMeshOperations.cpp
--- a/Engine/Source/Runtime/StaticMeshDescription/Private/StaticMeshOperations.cpp
+++ b/Engine/Source/Runtime/StaticMeshDescription/Private/StaticMeshOperations.cpp
@@ -1894,9 +1894,12 @@ bool FStaticMeshOperations::GenerateUniqueUVsForStaticMesh(const FMeshDescriptio
 				DuplicateMeshDescription.VertexInstanceAttributes().ForEach(
 					[&PolyHash, &RefVertexInstanceID](const FName AttributeName, auto AttributeArrayRef)
 					{
-						for (int32 Index = 0; Index < AttributeArrayRef.GetNumChannels(); ++Index)
+						for (int32 Channel = 0; Channel < AttributeArrayRef.GetNumChannels(); ++Channel)
 						{
-							PolyHash = HashCombine(PolyHash, GetTypeHash(AttributeArrayRef.Get(RefVertexInstanceID)));
+							for (const auto& Element : AttributeArrayRef.GetArrayView(RefVertexInstanceID, Channel))
+							{
+								PolyHash = HashCombine(PolyHash, GetTypeHash(Element));
+							}
 						}
 					}
 				);
