Implemented TArrayAttribute constructor for implicit conversion from non-const to const attribute types.

diff --git a/Engine/Source/Runtime/MeshDescription/Public/AttributeArrayContainer.h b/Engine/Source/Runtime/MeshDescription/Public/AttributeArrayContainer.h
--- a/Engine/Source/Runtime/MeshDescription/Public/AttributeArrayContainer.h
+++ b/Engine/Source/Runtime/MeshDescription/Public/AttributeArrayContainer.h
@@ -282,6 +282,7 @@ void TAttributeArrayContainer<AttributeType>::Remap(const TSparseArray<int32>& I
 template <typename AttributeType>
 class TArrayAttribute
 {
+	template <typename T> friend class TArrayAttribute;
 	using ArrayType = typename TCopyQualifiersFromTo<AttributeType, TAttributeArrayContainer<typename TRemoveCV<AttributeType>::Type>>::Type;
 
 public:
@@ -290,6 +291,15 @@ public:
 		  Index(InIndex)
 	{}
 
+	/**
+	 * Construct a TArrayAttribute<const T> from a TArrayAttribute<T>. 
+	 */
+	template <typename T = AttributeType, typename TEnableIf<TIsSame<T, const T>::Value, int>::Type = 0>
+	TArrayAttribute(TArrayAttribute<typename TRemoveCV<T>::Type> InValue)
+		: Array(InValue.Array),
+		  Index(InValue.Index)
+	{}
+
 	/**
 	 * Helper function for returning a typed pointer to the first array attribute entry.
 	 */
