Temporarily removed some deprecation directives while analysing a difference in compiler behaviour in order to remove build warnings.

diff --git a/Engine/Source/Runtime/MeshDescription/Public/MeshAttributeArray.h b/Engine/Source/Runtime/MeshDescription/Public/MeshAttributeArray.h
--- a/Engine/Source/Runtime/MeshDescription/Public/MeshAttributeArray.h
+++ b/Engine/Source/Runtime/MeshDescription/Public/MeshAttributeArray.h
@@ -280,13 +280,13 @@ public:
 	virtual void Serialize(FArchive& Ar) = 0;
 	virtual void Remap(const TSparseArray<int32>& IndexRemap) = 0;
 
-	UE_DEPRECATED(5.0, "Please use GetNumChannels().")
+//	UE_DEPRECATED(5.0, "Please use GetNumChannels().")
 	virtual int32 GetNumIndices() const = 0;
-	UE_DEPRECATED(5.0, "Please use SetNumChannels().")
+//	UE_DEPRECATED(5.0, "Please use SetNumChannels().")
 	virtual void SetNumIndices(const int32 NumIndices) = 0;
-	UE_DEPRECATED(5.0, "Please use InsertChannel().")
+//	UE_DEPRECATED(5.0, "Please use InsertChannel().")
 	virtual void InsertIndex(const int32 Index) = 0;
-	UE_DEPRECATED(5.0, "Please use RemoveChannel().")
+//	UE_DEPRECATED(5.0, "Please use RemoveChannel().")
 	virtual void RemoveIndex(const int32 Index) = 0;
 
 	virtual int32 GetNumChannels() const = 0;
