Support for the UMeshDescriptionBase / UStaticMeshDescription hierarchy to take a reference to an existing mesh description as well as hold one by value.
Added additional API for getting element counts.

diff --git a/Engine/Source/Runtime/Engine/Classes/Engine/StaticMesh.h b/Engine/Source/Runtime/Engine/Classes/Engine/StaticMesh.h
--- a/Engine/Source/Runtime/Engine/Classes/Engine/StaticMesh.h
+++ b/Engine/Source/Runtime/Engine/Classes/Engine/StaticMesh.h
@@ -1386,6 +1386,9 @@ public:
 	UFUNCTION(BlueprintCallable, Category="StaticMesh")
 	ENGINE_API void BuildFromStaticMeshDescriptions(const TArray<UStaticMeshDescription*>& StaticMeshDescriptions, bool bBuildSimpleCollision = false);
 
+	/** Return a new StaticMeshDescription referencing the MeshDescription of the given LOD */
+	UFUNCTION(BlueprintCallable, Category="StaticMesh")
+	ENGINE_API UStaticMeshDescription* GetStaticMeshDescription(int32 LODIndex);
 
 	 /** Structure that defines parameters passed into the build mesh description function */
 	struct FBuildMeshDescriptionsParams
diff --git a/Engine/Source/Runtime/Engine/Private/StaticMesh.cpp b/Engine/Source/Runtime/Engine/Private/StaticMesh.cpp
--- a/Engine/Source/Runtime/Engine/Private/StaticMesh.cpp
+++ b/Engine/Source/Runtime/Engine/Private/StaticMesh.cpp
@@ -5950,6 +5950,19 @@ UStaticMeshDescription* UStaticMesh::CreateStaticMeshDescription(UObject* Outer)
 }
 
 
+UStaticMeshDescription* UStaticMesh::GetStaticMeshDescription(int32 LODIndex)
+{
+	if (LODIndex < GetNumSourceModels())
+	{
+		UStaticMeshDescription* StaticMeshDescription = NewObject<UStaticMeshDescription>(this, NAME_None, RF_Transient);
+		StaticMeshDescription->SetMeshDescriptionRef(GetMeshDescription(LODIndex));
+		return StaticMeshDescription;
+	}
+
+	return nullptr;
+}
+
+
 void UStaticMesh::BuildFromStaticMeshDescriptions(const TArray<UStaticMeshDescription*>& StaticMeshDescriptions, bool bBuildSimpleCollision)
 {
 	TArray<const FMeshDescription*> MeshDescriptions;
diff --git a/Engine/Source/Runtime/MeshDescription/Private/MeshDescriptionBase.cpp b/Engine/Source/Runtime/MeshDescription/Private/MeshDescriptionBase.cpp
--- a/Engine/Source/Runtime/MeshDescription/Private/MeshDescriptionBase.cpp
+++ b/Engine/Source/Runtime/MeshDescription/Private/MeshDescriptionBase.cpp
@@ -12,41 +12,41 @@ void UMeshDescriptionBase::Serialize(FArchive& Ar)
 
 void UMeshDescriptionBase::RegisterAttributes()
 {
-	RequiredAttributes = MakeUnique<FMeshAttributes>(MeshDescription);
+	RequiredAttributes = MakeUnique<FMeshAttributes>(GetMeshDescription());
 	RequiredAttributes->Register();
 }
 
 void UMeshDescriptionBase::Reset()
 {
-	MeshDescription = FMeshDescription();
+	SetMeshDescriptionRef(nullptr);
 	RegisterAttributes();
 }
 
 void UMeshDescriptionBase::Empty()
 {
-	MeshDescription.Empty();
+	GetMeshDescription().Empty();
 }
 
 bool UMeshDescriptionBase::IsEmpty() const
 {
-	return MeshDescription.IsEmpty();
+	return GetMeshDescription().IsEmpty();
 }
 
 void UMeshDescriptionBase::ReserveNewVertices(int32 NumberOfNewVertices)
 {
-	MeshDescription.ReserveNewVertices(NumberOfNewVertices);
+	GetMeshDescription().ReserveNewVertices(NumberOfNewVertices);
 }
 
 FVertexID UMeshDescriptionBase::CreateVertex()
 {
-	return MeshDescription.CreateVertex();
+	return GetMeshDescription().CreateVertex();
 }
 
 void UMeshDescriptionBase::CreateVertexWithID(FVertexID VertexID)
 {
-	if (!MeshDescription.IsVertexValid(VertexID))
+	if (!GetMeshDescription().IsVertexValid(VertexID))
 	{
-		MeshDescription.CreateVertexWithID(VertexID);
+		GetMeshDescription().CreateVertexWithID(VertexID);
 	}
 	else
 	{
@@ -56,9 +56,9 @@ void UMeshDescriptionBase::CreateVertexWithID(FVertexID VertexID)
 
 void UMeshDescriptionBase::DeleteVertex(FVertexID VertexID)
 {
-	if (MeshDescription.IsVertexValid(VertexID))
+	if (GetMeshDescription().IsVertexValid(VertexID))
 	{
-		MeshDescription.DeleteVertex(VertexID);
+		GetMeshDescription().DeleteVertex(VertexID);
 	}
 	else
 	{
@@ -68,19 +68,19 @@ void UMeshDescriptionBase::DeleteVertex(FVertexID VertexID)
 
 bool UMeshDescriptionBase::IsVertexValid(FVertexID VertexID) const
 {
-	return MeshDescription.IsVertexValid(VertexID);
+	return GetMeshDescription().IsVertexValid(VertexID);
 }
 
 void UMeshDescriptionBase::ReserveNewVertexInstances(int32 NumberOfNewVertexInstances)
 {
-	MeshDescription.ReserveNewEdges(NumberOfNewVertexInstances);
+	GetMeshDescription().ReserveNewEdges(NumberOfNewVertexInstances);
 }
 
 FVertexInstanceID UMeshDescriptionBase::CreateVertexInstance(FVertexID VertexID)
 {
-	if (MeshDescription.IsVertexValid(VertexID))
+	if (GetMeshDescription().IsVertexValid(VertexID))
 	{
-		return MeshDescription.CreateVertexInstance(VertexID);
+		return GetMeshDescription().CreateVertexInstance(VertexID);
 	}
 	else
 	{
@@ -91,11 +91,11 @@ FVertexInstanceID UMeshDescriptionBase::CreateVertexInstance(FVertexID VertexID)
 
 void UMeshDescriptionBase::CreateVertexInstanceWithID(FVertexInstanceID VertexInstanceID, FVertexID VertexID)
 {
-	if (MeshDescription.IsVertexValid(VertexID))
+	if (GetMeshDescription().IsVertexValid(VertexID))
 	{
-		if (!MeshDescription.IsVertexInstanceValid(VertexInstanceID))
+		if (!GetMeshDescription().IsVertexInstanceValid(VertexInstanceID))
 		{
-			MeshDescription.CreateVertexInstanceWithID(VertexInstanceID, VertexID);
+			GetMeshDescription().CreateVertexInstanceWithID(VertexInstanceID, VertexID);
 		}
 		else
 		{
@@ -110,9 +110,9 @@ void UMeshDescriptionBase::CreateVertexInstanceWithID(FVertexInstanceID VertexIn
 
 void UMeshDescriptionBase::DeleteVertexInstance(FVertexInstanceID VertexInstanceID, TArray<FVertexID>& OrphanedVertices)
 {
-	if (MeshDescription.IsVertexInstanceValid(VertexInstanceID))
+	if (GetMeshDescription().IsVertexInstanceValid(VertexInstanceID))
 	{
-		MeshDescription.DeleteVertexInstance(VertexInstanceID, &OrphanedVertices);
+		GetMeshDescription().DeleteVertexInstance(VertexInstanceID, &OrphanedVertices);
 	}
 	else
 	{
@@ -122,59 +122,59 @@ void UMeshDescriptionBase::DeleteVertexInstance(FVertexInstanceID VertexInstance
 
 bool UMeshDescriptionBase::IsVertexInstanceValid(FVertexInstanceID VertexInstanceID) const
 {
-	return MeshDescription.IsVertexInstanceValid(VertexInstanceID);
+	return GetMeshDescription().IsVertexInstanceValid(VertexInstanceID);
 }
 
 void UMeshDescriptionBase::ReserveNewEdges(int32 NumberOfNewEdges)
 {
-	MeshDescription.ReserveNewEdges(NumberOfNewEdges);
+	GetMeshDescription().ReserveNewEdges(NumberOfNewEdges);
 }
 
 FEdgeID UMeshDescriptionBase::CreateEdge(FVertexID VertexID0, FVertexID VertexID1)
 {
-	if (!MeshDescription.IsVertexValid(VertexID0))
+	if (!GetMeshDescription().IsVertexValid(VertexID0))
 	{
 		UE_LOG(LogMeshDescription, Warning, TEXT("CreateEdge: VertexID %d doesn't exist."), VertexID0.GetValue());
 		return INDEX_NONE;
 	}
 
-	if (!MeshDescription.IsVertexValid(VertexID1))
+	if (!GetMeshDescription().IsVertexValid(VertexID1))
 	{
 		UE_LOG(LogMeshDescription, Warning, TEXT("CreateEdge: VertexID %d doesn't exist."), VertexID1.GetValue());
 		return INDEX_NONE;
 	}
 
-	return MeshDescription.CreateEdge(VertexID0, VertexID1);
+	return GetMeshDescription().CreateEdge(VertexID0, VertexID1);
 }
 
 void UMeshDescriptionBase::CreateEdgeWithID(FEdgeID EdgeID, FVertexID VertexID0, FVertexID VertexID1)
 {
-	if (MeshDescription.IsEdgeValid(EdgeID))
+	if (GetMeshDescription().IsEdgeValid(EdgeID))
 	{
 		UE_LOG(LogMeshDescription, Warning, TEXT("CreateEdgeWithID: EdgeID %d already exists."), EdgeID.GetValue());
 		return;
 	}
 
-	if (!MeshDescription.IsVertexValid(VertexID0))
+	if (!GetMeshDescription().IsVertexValid(VertexID0))
 	{
 		UE_LOG(LogMeshDescription, Warning, TEXT("CreateEdgeWithID: VertexID %d doesn't exist."), VertexID0.GetValue());
 		return;
 	}
 
-	if (!MeshDescription.IsVertexValid(VertexID1))
+	if (!GetMeshDescription().IsVertexValid(VertexID1))
 	{
 		UE_LOG(LogMeshDescription, Warning, TEXT("CreateEdgeWithID: VertexID %d doesn't exist."), VertexID1.GetValue());
 		return;
 	}
 
-	MeshDescription.CreateEdgeWithID(EdgeID, VertexID0, VertexID1);
+	GetMeshDescription().CreateEdgeWithID(EdgeID, VertexID0, VertexID1);
 }
 
 void UMeshDescriptionBase::DeleteEdge(FEdgeID EdgeID, TArray<FVertexID>& OrphanedVertices)
 {
-	if (MeshDescription.IsEdgeValid(EdgeID))
+	if (GetMeshDescription().IsEdgeValid(EdgeID))
 	{
-		MeshDescription.DeleteEdge(EdgeID, &OrphanedVertices);
+		GetMeshDescription().DeleteEdge(EdgeID, &OrphanedVertices);
 	}
 	else
 	{
@@ -184,142 +184,142 @@ void UMeshDescriptionBase::DeleteEdge(FEdgeID EdgeID, TArray<FVertexID>& Orphane
 
 bool UMeshDescriptionBase::IsEdgeValid(FEdgeID EdgeID) const
 {
-	return MeshDescription.IsEdgeValid(EdgeID);
+	return GetMeshDescription().IsEdgeValid(EdgeID);
 }
 
 void UMeshDescriptionBase::ReserveNewTriangles(int32 NumberOfNewTriangles)
 {
-	MeshDescription.ReserveNewTriangles(NumberOfNewTriangles);
+	GetMeshDescription().ReserveNewTriangles(NumberOfNewTriangles);
 }
 
 FTriangleID UMeshDescriptionBase::CreateTriangle(FPolygonGroupID PolygonGroupID, const TArray<FVertexInstanceID>& VertexInstanceIDs, TArray<FEdgeID>& NewEdgeIDs)
 {
-	if (!MeshDescription.IsPolygonGroupValid(PolygonGroupID))
+	if (!GetMeshDescription().IsPolygonGroupValid(PolygonGroupID))
 	{
 		UE_LOG(LogMeshDescription, Warning, TEXT("CreateTriangle: PolygonGroupID %d doesn't exist."), PolygonGroupID.GetValue());
 		return INDEX_NONE;
 	}
 
-	return MeshDescription.CreateTriangle(PolygonGroupID, VertexInstanceIDs, &NewEdgeIDs);
+	return GetMeshDescription().CreateTriangle(PolygonGroupID, VertexInstanceIDs, &NewEdgeIDs);
 }
 
 void UMeshDescriptionBase::CreateTriangleWithID(FTriangleID TriangleID, FPolygonGroupID PolygonGroupID, const TArray<FVertexInstanceID>& VertexInstanceIDs, TArray<FEdgeID>& NewEdgeIDs)
 {
-	if (MeshDescription.IsTriangleValid(TriangleID))
+	if (GetMeshDescription().IsTriangleValid(TriangleID))
 	{
 		UE_LOG(LogMeshDescription, Warning, TEXT("CreateTriangleWithID: TriangleID %d already exists."), TriangleID.GetValue());
 		return;
 	}
 
-	if (!MeshDescription.IsPolygonGroupValid(PolygonGroupID))
+	if (!GetMeshDescription().IsPolygonGroupValid(PolygonGroupID))
 	{
 		UE_LOG(LogMeshDescription, Warning, TEXT("CreateTriangleWithID: PolygonGroupID %d doesn't exist."), PolygonGroupID.GetValue());
 		return;
 	}
 
-	MeshDescription.CreateTriangleWithID(TriangleID, PolygonGroupID, VertexInstanceIDs, &NewEdgeIDs);
+	GetMeshDescription().CreateTriangleWithID(TriangleID, PolygonGroupID, VertexInstanceIDs, &NewEdgeIDs);
 }
 
 void UMeshDescriptionBase::DeleteTriangle(FTriangleID TriangleID, TArray<FEdgeID>& OrphanedEdges, TArray<FVertexInstanceID>& OrphanedVertexInstances, TArray<FPolygonGroupID>& OrphanedPolygonGroups)
 {
-	if (!MeshDescription.IsTriangleValid(TriangleID))
+	if (!GetMeshDescription().IsTriangleValid(TriangleID))
 	{
 		UE_LOG(LogMeshDescription, Warning, TEXT("DeleteTriangle: TriangleID %d doesn't exist."), TriangleID.GetValue());
 		return;
 	}
 
-	MeshDescription.DeleteTriangle(TriangleID, &OrphanedEdges, &OrphanedVertexInstances, &OrphanedPolygonGroups);
+	GetMeshDescription().DeleteTriangle(TriangleID, &OrphanedEdges, &OrphanedVertexInstances, &OrphanedPolygonGroups);
 }
 
 bool UMeshDescriptionBase::IsTriangleValid(const FTriangleID TriangleID) const
 {
-	return MeshDescription.IsTriangleValid(TriangleID);
+	return GetMeshDescription().IsTriangleValid(TriangleID);
 }
 
 void UMeshDescriptionBase::ReserveNewPolygons(const int32 NumberOfNewPolygons)
 {
-	MeshDescription.ReserveNewPolygons(NumberOfNewPolygons);
+	GetMeshDescription().ReserveNewPolygons(NumberOfNewPolygons);
 }
 
 FPolygonID UMeshDescriptionBase::CreatePolygon(FPolygonGroupID PolygonGroupID, TArray<FVertexInstanceID>& VertexInstanceIDs, TArray<FEdgeID>& NewEdgeIDs)
 {
-	if (!MeshDescription.IsPolygonGroupValid(PolygonGroupID))
+	if (!GetMeshDescription().IsPolygonGroupValid(PolygonGroupID))
 	{
 		UE_LOG(LogMeshDescription, Warning, TEXT("CreatePolygon: PolygonGroupID %d doesn't exist."), PolygonGroupID.GetValue());
 		return INDEX_NONE;
 	}
 
-	return MeshDescription.CreatePolygon(PolygonGroupID, VertexInstanceIDs, &NewEdgeIDs);
+	return GetMeshDescription().CreatePolygon(PolygonGroupID, VertexInstanceIDs, &NewEdgeIDs);
 }
 
 void UMeshDescriptionBase::CreatePolygonWithID(FPolygonID PolygonID, FPolygonGroupID PolygonGroupID, TArray<FVertexInstanceID>& VertexInstanceIDs, TArray<FEdgeID>& NewEdgeIDs)
 {
-	if (MeshDescription.IsPolygonValid(PolygonID))
+	if (GetMeshDescription().IsPolygonValid(PolygonID))
 	{
 		UE_LOG(LogMeshDescription, Warning, TEXT("CreatePolygonWithID: PolygonID %d already exists."), PolygonID.GetValue());
 		return;
 	}
 
-	if (!MeshDescription.IsPolygonGroupValid(PolygonGroupID))
+	if (!GetMeshDescription().IsPolygonGroupValid(PolygonGroupID))
 	{
 		UE_LOG(LogMeshDescription, Warning, TEXT("CreatePolygonWithID: PolygonGroupID %d doesn't exist."), PolygonGroupID.GetValue());
 		return;
 	}
 
-	MeshDescription.CreatePolygonWithID(PolygonID, PolygonGroupID, VertexInstanceIDs, &NewEdgeIDs);
+	GetMeshDescription().CreatePolygonWithID(PolygonID, PolygonGroupID, VertexInstanceIDs, &NewEdgeIDs);
 }
 
 void UMeshDescriptionBase::DeletePolygon(FPolygonID PolygonID, TArray<FEdgeID>& OrphanedEdges, TArray<FVertexInstanceID>& OrphanedVertexInstances, TArray<FPolygonGroupID>& OrphanedPolygonGroups)
 {
-	if (!MeshDescription.IsPolygonValid(PolygonID))
+	if (!GetMeshDescription().IsPolygonValid(PolygonID))
 	{
 		UE_LOG(LogMeshDescription, Warning, TEXT("DeletePolygon: PolygonID %d doesn't exist."), PolygonID.GetValue());
 		return;
 	}
 
-	MeshDescription.DeletePolygon(PolygonID, &OrphanedEdges, &OrphanedVertexInstances, &OrphanedPolygonGroups);
+	GetMeshDescription().DeletePolygon(PolygonID, &OrphanedEdges, &OrphanedVertexInstances, &OrphanedPolygonGroups);
 }
 
 bool UMeshDescriptionBase::IsPolygonValid(FPolygonID PolygonID) const
 {
-	return MeshDescription.IsPolygonValid(PolygonID);
+	return GetMeshDescription().IsPolygonValid(PolygonID);
 }
 
 void UMeshDescriptionBase::ReserveNewPolygonGroups(int32 NumberOfNewPolygonGroups)
 {
-	MeshDescription.ReserveNewPolygonGroups(NumberOfNewPolygonGroups);
+	GetMeshDescription().ReserveNewPolygonGroups(NumberOfNewPolygonGroups);
 }
 
 FPolygonGroupID UMeshDescriptionBase::CreatePolygonGroup()
 {
-	return MeshDescription.CreatePolygonGroup();
+	return GetMeshDescription().CreatePolygonGroup();
 }
 
 void UMeshDescriptionBase::CreatePolygonGroupWithID(FPolygonGroupID PolygonGroupID)
 {
-	if (MeshDescription.IsPolygonGroupValid(PolygonGroupID))
+	if (GetMeshDescription().IsPolygonGroupValid(PolygonGroupID))
 	{
 		UE_LOG(LogMeshDescription, Warning, TEXT("CreatePolygonGroupWithID: PolygonGroupID %d already exists."), PolygonGroupID.GetValue());
 		return;
 	}
 
-	MeshDescription.CreatePolygonGroupWithID(PolygonGroupID);
+	GetMeshDescription().CreatePolygonGroupWithID(PolygonGroupID);
 }
 
 void UMeshDescriptionBase::DeletePolygonGroup(FPolygonGroupID PolygonGroupID)
 {
-	if (!MeshDescription.IsPolygonGroupValid(PolygonGroupID))
+	if (!GetMeshDescription().IsPolygonGroupValid(PolygonGroupID))
 	{
 		UE_LOG(LogMeshDescription, Warning, TEXT("DeletePolygonGroup: FPolygonGroupID %d doesn't exist."), PolygonGroupID.GetValue());
 		return;
 	}
 
-	MeshDescription.DeletePolygonGroup(PolygonGroupID);
+	GetMeshDescription().DeletePolygonGroup(PolygonGroupID);
 }
 
 bool UMeshDescriptionBase::IsPolygonGroupValid(FPolygonGroupID PolygonGroupID) const
 {
-	return MeshDescription.IsPolygonGroupValid(PolygonGroupID);
+	return GetMeshDescription().IsPolygonGroupValid(PolygonGroupID);
 }
 
 
@@ -328,163 +328,163 @@ bool UMeshDescriptionBase::IsPolygonGroupValid(FPolygonGroupID PolygonGroupID) c
 
 bool UMeshDescriptionBase::IsVertexOrphaned(FVertexID VertexID) const
 {
-	if (!MeshDescription.IsVertexValid(VertexID))
+	if (!GetMeshDescription().IsVertexValid(VertexID))
 	{
 		UE_LOG(LogMeshDescription, Warning, TEXT("IsVertexOrphaned: VertexID %d doesn't exist."), VertexID.GetValue());
 		return false;
 	}
 
-	return MeshDescription.IsVertexOrphaned(VertexID);
+	return GetMeshDescription().IsVertexOrphaned(VertexID);
 }
 
 FEdgeID UMeshDescriptionBase::GetVertexPairEdge(FVertexID VertexID0, FVertexID VertexID1) const
 {
-	if (!MeshDescription.IsVertexValid(VertexID0))
+	if (!GetMeshDescription().IsVertexValid(VertexID0))
 	{
 		UE_LOG(LogMeshDescription, Warning, TEXT("GetVertexPairEdge: VertexID %d doesn't exist."), VertexID0.GetValue());
 		return INDEX_NONE;
 	}
 
-	if (!MeshDescription.IsVertexValid(VertexID1))
+	if (!GetMeshDescription().IsVertexValid(VertexID1))
 	{
 		UE_LOG(LogMeshDescription, Warning, TEXT("GetVertexPairEdge: VertexID %d doesn't exist."), VertexID1.GetValue());
 		return INDEX_NONE;
 	}
 
-	return MeshDescription.GetVertexPairEdge(VertexID0, VertexID1);
+	return GetMeshDescription().GetVertexPairEdge(VertexID0, VertexID1);
 }
 
 void UMeshDescriptionBase::GetVertexConnectedEdges(FVertexID VertexID, TArray<FEdgeID>& OutEdgeIDs) const
 {
-	if (!MeshDescription.IsVertexValid(VertexID))
+	if (!GetMeshDescription().IsVertexValid(VertexID))
 	{
 		UE_LOG(LogMeshDescription, Warning, TEXT("GetVertexConnectedEdges: VertexID %d doesn't exist."), VertexID.GetValue());
 		return;
 	}
 
-	OutEdgeIDs = MeshDescription.GetVertexConnectedEdgeIDs(VertexID);
+	OutEdgeIDs = GetMeshDescription().GetVertexConnectedEdgeIDs(VertexID);
 }
 
 int32 UMeshDescriptionBase::GetNumVertexConnectedEdges(FVertexID VertexID) const
 {
-	if (!MeshDescription.IsVertexValid(VertexID))
+	if (!GetMeshDescription().IsVertexValid(VertexID))
 	{
 		UE_LOG(LogMeshDescription, Warning, TEXT("GetNumVertexConnectedEdges: VertexID %d doesn't exist."), VertexID.GetValue());
 		return 0;
 	}
 
-	return MeshDescription.GetNumVertexConnectedEdges(VertexID);
+	return GetMeshDescription().GetNumVertexConnectedEdges(VertexID);
 }
 
 void UMeshDescriptionBase::GetVertexVertexInstances(FVertexID VertexID, TArray<FVertexInstanceID>& OutVertexInstanceIDs) const
 {
-	if (!MeshDescription.IsVertexValid(VertexID))
+	if (!GetMeshDescription().IsVertexValid(VertexID))
 	{
 		UE_LOG(LogMeshDescription, Warning, TEXT("GetVertexVertexInstances: VertexID %d doesn't exist."), VertexID.GetValue());
 		return;
 	}
 
-	OutVertexInstanceIDs = MeshDescription.GetVertexVertexInstanceIDs(VertexID);
+	OutVertexInstanceIDs = GetMeshDescription().GetVertexVertexInstanceIDs(VertexID);
 }
 
 int32 UMeshDescriptionBase::GetNumVertexVertexInstances(FVertexID VertexID) const
 {
-	if (!MeshDescription.IsVertexValid(VertexID))
+	if (!GetMeshDescription().IsVertexValid(VertexID))
 	{
 		UE_LOG(LogMeshDescription, Warning, TEXT("GetNumVertexVertexInstances: VertexID %d doesn't exist."), VertexID.GetValue());
 		return 0;
 	}
 
-	return MeshDescription.GetNumVertexVertexInstances(VertexID);
+	return GetMeshDescription().GetNumVertexVertexInstances(VertexID);
 }
 
 void UMeshDescriptionBase::GetVertexConnectedTriangles(FVertexID VertexID, TArray<FTriangleID>& OutConnectedTriangleIDs) const
 {
-	if (!MeshDescription.IsVertexValid(VertexID))
+	if (!GetMeshDescription().IsVertexValid(VertexID))
 	{
 		UE_LOG(LogMeshDescription, Warning, TEXT("GetVertexConnectedTriangles: VertexID %d doesn't exist."), VertexID.GetValue());
 		return;
 	}
 
-	MeshDescription.GetVertexConnectedTriangles(VertexID, OutConnectedTriangleIDs);
+	GetMeshDescription().GetVertexConnectedTriangles(VertexID, OutConnectedTriangleIDs);
 }
 
 int32 UMeshDescriptionBase::GetNumVertexConnectedTriangles(FVertexID VertexID) const
 {
-	if (!MeshDescription.IsVertexValid(VertexID))
+	if (!GetMeshDescription().IsVertexValid(VertexID))
 	{
 		UE_LOG(LogMeshDescription, Warning, TEXT("GetNumVertexConnectedTriangles: VertexID %d doesn't exist."), VertexID.GetValue());
 		return 0;
 	}
 
-	return MeshDescription.GetNumVertexConnectedTriangles(VertexID);
+	return GetMeshDescription().GetNumVertexConnectedTriangles(VertexID);
 }
 
 void UMeshDescriptionBase::GetVertexConnectedPolygons(FVertexID VertexID, TArray<FPolygonID>& OutConnectedPolygonIDs) const
 {
-	if (!MeshDescription.IsVertexValid(VertexID))
+	if (!GetMeshDescription().IsVertexValid(VertexID))
 	{
 		UE_LOG(LogMeshDescription, Warning, TEXT("GetVertexConnectedPolygons: VertexID %d doesn't exist."), VertexID.GetValue());
 		return;
 	}
 
-	MeshDescription.GetVertexConnectedPolygons(VertexID, OutConnectedPolygonIDs);
+	GetMeshDescription().GetVertexConnectedPolygons(VertexID, OutConnectedPolygonIDs);
 }
 
 int32 UMeshDescriptionBase::GetNumVertexConnectedPolygons(FVertexID VertexID) const
 {
-	if (!MeshDescription.IsVertexValid(VertexID))
+	if (!GetMeshDescription().IsVertexValid(VertexID))
 	{
 		UE_LOG(LogMeshDescription, Warning, TEXT("GetNumVertexConnectedPolygons: VertexID %d doesn't exist."), VertexID.GetValue());
 		return 0;
 	}
 
-	return MeshDescription.GetNumVertexConnectedPolygons(VertexID);
+	return GetMeshDescription().GetNumVertexConnectedPolygons(VertexID);
 }
 
 void UMeshDescriptionBase::GetVertexAdjacentVertices(FVertexID VertexID, TArray<FVertexID>& OutAdjacentVertexIDs) const
 {
-	if (!MeshDescription.IsVertexValid(VertexID))
+	if (!GetMeshDescription().IsVertexValid(VertexID))
 	{
 		UE_LOG(LogMeshDescription, Warning, TEXT("GetVertexAdjacentVertices: VertexID %d doesn't exist."), VertexID.GetValue());
 		return;
 	}
 
-	MeshDescription.GetVertexAdjacentVertices(VertexID, OutAdjacentVertexIDs);
+	GetMeshDescription().GetVertexAdjacentVertices(VertexID, OutAdjacentVertexIDs);
 }
 
 FVector UMeshDescriptionBase::GetVertexPosition(FVertexID VertexID) const
 {
-	if (!MeshDescription.IsVertexValid(VertexID))
+	if (!GetMeshDescription().IsVertexValid(VertexID))
 	{
 		UE_LOG(LogMeshDescription, Warning, TEXT("GetVertexAttribute: VertexID %d doesn't exist."), VertexID.GetValue());
 		return FVector::ZeroVector;
 	}
 
-	if (!MeshDescription.VertexAttributes().HasAttribute(MeshAttribute::Vertex::Position))
+	if (!GetMeshDescription().VertexAttributes().HasAttribute(MeshAttribute::Vertex::Position))
 	{
 		UE_LOG(LogMeshDescription, Warning, TEXT("GetVertexAttribute: VertexAttribute Position doesn't exist."));
 		return FVector::ZeroVector;
 	}
 
-	return MeshDescription.VertexAttributes().GetAttribute<FVector>(VertexID, MeshAttribute::Vertex::Position);
+	return GetMeshDescription().VertexAttributes().GetAttribute<FVector>(VertexID, MeshAttribute::Vertex::Position);
 }
 
 void UMeshDescriptionBase::SetVertexPosition(FVertexID VertexID, const FVector& Position)
 {
-	if (!MeshDescription.IsVertexValid(VertexID))
+	if (!GetMeshDescription().IsVertexValid(VertexID))
 	{
 		UE_LOG(LogMeshDescription, Warning, TEXT("SetVertexAttribute: VertexID %d doesn't exist."), VertexID.GetValue());
 		return;
 	}
 
-	if (!MeshDescription.VertexAttributes().HasAttribute(MeshAttribute::Vertex::Position))
+	if (!GetMeshDescription().VertexAttributes().HasAttribute(MeshAttribute::Vertex::Position))
 	{
 		UE_LOG(LogMeshDescription, Warning, TEXT("SetVertexAttribute: VertexAttribute Position doesn't exist."));
 		return;
 	}
 
-	MeshDescription.VertexAttributes().SetAttribute(VertexID, MeshAttribute::Vertex::Position, 0, Position);
+	GetMeshDescription().VertexAttributes().SetAttribute(VertexID, MeshAttribute::Vertex::Position, 0, Position);
 }
 
 
@@ -493,74 +493,74 @@ void UMeshDescriptionBase::SetVertexPosition(FVertexID VertexID, const FVector&
 
 FVertexID UMeshDescriptionBase::GetVertexInstanceVertex(FVertexInstanceID VertexInstanceID) const
 {
-	if (!MeshDescription.IsVertexInstanceValid(VertexInstanceID))
+	if (!GetMeshDescription().IsVertexInstanceValid(VertexInstanceID))
 	{
 		UE_LOG(LogMeshDescription, Warning, TEXT("GetVertexInstanceVertex: VertexInstanceID %d doesn't exist."), VertexInstanceID.GetValue());
 		return INDEX_NONE;
 	}
 
-	return MeshDescription.GetVertexInstanceVertex(VertexInstanceID);
+	return GetMeshDescription().GetVertexInstanceVertex(VertexInstanceID);
 }
 
 FEdgeID UMeshDescriptionBase::GetVertexInstancePairEdge(FVertexInstanceID VertexInstanceID0, FVertexInstanceID VertexInstanceID1) const
 {
-	if (!MeshDescription.IsVertexInstanceValid(VertexInstanceID0))
+	if (!GetMeshDescription().IsVertexInstanceValid(VertexInstanceID0))
 	{
 		UE_LOG(LogMeshDescription, Warning, TEXT("GetVertexInstancePairEdge: VertexInstanceID %d doesn't exist."), VertexInstanceID0.GetValue());
 		return INDEX_NONE;
 	}
 
-	if (!MeshDescription.IsVertexInstanceValid(VertexInstanceID1))
+	if (!GetMeshDescription().IsVertexInstanceValid(VertexInstanceID1))
 	{
 		UE_LOG(LogMeshDescription, Warning, TEXT("GetVertexInstancePairEdge: VertexInstanceID %d doesn't exist."), VertexInstanceID1.GetValue());
 		return INDEX_NONE;
 	}
 
-	return MeshDescription.GetVertexInstancePairEdge(VertexInstanceID0, VertexInstanceID1);
+	return GetMeshDescription().GetVertexInstancePairEdge(VertexInstanceID0, VertexInstanceID1);
 }
 
 void UMeshDescriptionBase::GetVertexInstanceConnectedTriangles(FVertexInstanceID VertexInstanceID, TArray<FTriangleID>& OutConnectedTriangleIDs) const
 {
-	if (!MeshDescription.IsVertexInstanceValid(VertexInstanceID))
+	if (!GetMeshDescription().IsVertexInstanceValid(VertexInstanceID))
 	{
 		UE_LOG(LogMeshDescription, Warning, TEXT("GetVertexInstanceConnectedTriangles: VertexInstanceID %d doesn't exist."), VertexInstanceID.GetValue());
 		return;
 	}
 
-	OutConnectedTriangleIDs = MeshDescription.GetVertexInstanceConnectedTriangleIDs(VertexInstanceID);
+	OutConnectedTriangleIDs = GetMeshDescription().GetVertexInstanceConnectedTriangleIDs(VertexInstanceID);
 }
 
 int32 UMeshDescriptionBase::GetNumVertexInstanceConnectedTriangles(FVertexInstanceID VertexInstanceID) const
 {
-	if (!MeshDescription.IsVertexInstanceValid(VertexInstanceID))
+	if (!GetMeshDescription().IsVertexInstanceValid(VertexInstanceID))
 	{
 		UE_LOG(LogMeshDescription, Warning, TEXT("GetNumVertexInstanceConnectedTriangles: VertexInstanceID %d doesn't exist."), VertexInstanceID.GetValue());
 		return 0;
 	}
 
-	return MeshDescription.GetNumVertexInstanceConnectedTriangles(VertexInstanceID);
+	return GetMeshDescription().GetNumVertexInstanceConnectedTriangles(VertexInstanceID);
 }
 
 void UMeshDescriptionBase::GetVertexInstanceConnectedPolygons(FVertexInstanceID VertexInstanceID, TArray<FPolygonID>& OutConnectedPolygonIDs) const
 {
-	if (!MeshDescription.IsVertexInstanceValid(VertexInstanceID))
+	if (!GetMeshDescription().IsVertexInstanceValid(VertexInstanceID))
 	{
 		UE_LOG(LogMeshDescription, Warning, TEXT("GetVertexInstanceConnectedPolygons: VertexInstanceID %d doesn't exist."), VertexInstanceID.GetValue());
 		return;
 	}
 
-	MeshDescription.GetVertexInstanceConnectedPolygons(VertexInstanceID, OutConnectedPolygonIDs);
+	GetMeshDescription().GetVertexInstanceConnectedPolygons(VertexInstanceID, OutConnectedPolygonIDs);
 }
 
 int32 UMeshDescriptionBase::GetNumVertexInstanceConnectedPolygons(FVertexInstanceID VertexInstanceID) const
 {
-	if (!MeshDescription.IsVertexInstanceValid(VertexInstanceID))
+	if (!GetMeshDescription().IsVertexInstanceValid(VertexInstanceID))
 	{
 		UE_LOG(LogMeshDescription, Warning, TEXT("GetNumVertexInstanceConnectedPolygons: VertexInstanceID %d doesn't exist."), VertexInstanceID.GetValue());
 		return 0;
 	}
 
-	return MeshDescription.GetNumVertexInstanceConnectedPolygons(VertexInstanceID);
+	return GetMeshDescription().GetNumVertexInstanceConnectedPolygons(VertexInstanceID);
 }
 
 
@@ -569,79 +569,79 @@ int32 UMeshDescriptionBase::GetNumVertexInstanceConnectedPolygons(FVertexInstanc
 
 bool UMeshDescriptionBase::IsEdgeInternal(FEdgeID EdgeID) const
 {
-	if (!MeshDescription.IsEdgeValid(EdgeID))
+	if (!GetMeshDescription().IsEdgeValid(EdgeID))
 	{
 		UE_LOG(LogMeshDescription, Warning, TEXT("IsEdgeInternal: EdgeID %d doesn't exist."), EdgeID.GetValue());
 		return false;
 	}
 
-	return MeshDescription.IsEdgeInternal(EdgeID);
+	return GetMeshDescription().IsEdgeInternal(EdgeID);
 }
 
 bool UMeshDescriptionBase::IsEdgeInternalToPolygon(FEdgeID EdgeID, FPolygonID PolygonID) const
 {
-	if (!MeshDescription.IsEdgeValid(EdgeID))
+	if (!GetMeshDescription().IsEdgeValid(EdgeID))
 	{
 		UE_LOG(LogMeshDescription, Warning, TEXT("IsEdgeInternalToPolygon: EdgeID %d doesn't exist."), EdgeID.GetValue());
 		return false;
 	}
 
-	if (!MeshDescription.IsPolygonValid(PolygonID))
+	if (!GetMeshDescription().IsPolygonValid(PolygonID))
 	{
 		UE_LOG(LogMeshDescription, Warning, TEXT("IsEdgeInternalToPolygon: PolygonID %d doesn't exist."), PolygonID.GetValue());
 		return false;
 	}
 
-	return MeshDescription.IsEdgeInternalToPolygon(EdgeID, PolygonID);
+	return GetMeshDescription().IsEdgeInternalToPolygon(EdgeID, PolygonID);
 }
 
 void UMeshDescriptionBase::GetEdgeConnectedTriangles(FEdgeID EdgeID, TArray<FTriangleID>& OutConnectedTriangleIDs) const
 {
-	if (!MeshDescription.IsEdgeValid(EdgeID))
+	if (!GetMeshDescription().IsEdgeValid(EdgeID))
 	{
 		UE_LOG(LogMeshDescription, Warning, TEXT("GetEdgeConnectedTriangles: EdgeID %d doesn't exist."), EdgeID.GetValue());
 		return;
 	}
 
-	OutConnectedTriangleIDs = MeshDescription.GetEdgeConnectedTriangleIDs(EdgeID);
+	OutConnectedTriangleIDs = GetMeshDescription().GetEdgeConnectedTriangleIDs(EdgeID);
 }
 
 int32 UMeshDescriptionBase::GetNumEdgeConnectedTriangles(FEdgeID EdgeID) const
 {
-	if (!MeshDescription.IsEdgeValid(EdgeID))
+	if (!GetMeshDescription().IsEdgeValid(EdgeID))
 	{
 		UE_LOG(LogMeshDescription, Warning, TEXT("GetNumEdgeConnectedTriangles: EdgeID %d doesn't exist."), EdgeID.GetValue());
 		return 0;
 	}
 
-	return MeshDescription.GetNumEdgeConnectedTriangles(EdgeID);
+	return GetMeshDescription().GetNumEdgeConnectedTriangles(EdgeID);
 }
 
 void UMeshDescriptionBase::GetEdgeConnectedPolygons(FEdgeID EdgeID, TArray<FPolygonID>& OutConnectedPolygonIDs) const
 {
-	if (!MeshDescription.IsEdgeValid(EdgeID))
+	if (!GetMeshDescription().IsEdgeValid(EdgeID))
 	{
 		UE_LOG(LogMeshDescription, Warning, TEXT("GetEdgeConnectedPolygons: EdgeID %d doesn't exist."), EdgeID.GetValue());
 		return;
 	}
 
-	MeshDescription.GetEdgeConnectedPolygons(EdgeID, OutConnectedPolygonIDs);
+	GetMeshDescription().GetEdgeConnectedPolygons(EdgeID, OutConnectedPolygonIDs);
 }
 
 int32 UMeshDescriptionBase::GetNumEdgeConnectedPolygons(FEdgeID EdgeID) const
 {
-	if (!MeshDescription.IsEdgeValid(EdgeID))
+	if (!GetMeshDescription().IsEdgeValid(EdgeID))
 	{
 		UE_LOG(LogMeshDescription, Warning, TEXT("GetNumEdgeConnectedPolygons: EdgeID %d doesn't exist."), EdgeID.GetValue());
 		return 0;
 	}
 
-	return MeshDescription.GetNumEdgeConnectedPolygons(EdgeID);
+	return GetMeshDescription().GetNumEdgeConnectedPolygons(EdgeID);
 }
 
 FVertexID UMeshDescriptionBase::GetEdgeVertex(FEdgeID EdgeID, int32 VertexNumber) const
 {
-	if (!MeshDescription.IsEdgeValid(EdgeID))
+	if (!GetMeshDescription().IsEdgeValid(EdgeID))
 	{
 		UE_LOG(LogMeshDescription, Warning, TEXT("GetEdgeVertex: EdgeID %d doesn't exist."), EdgeID.GetValue());
 		return INDEX_NONE;
@@ -653,18 +653,18 @@ FVertexID UMeshDescriptionBase::GetEdgeVertex(FEdgeID EdgeID, int32 VertexNumber
 		return INDEX_NONE;
 	}
 
-	return MeshDescription.GetEdgeVertex(EdgeID, VertexNumber);
+	return GetMeshDescription().GetEdgeVertex(EdgeID, VertexNumber);
 }
 
 void UMeshDescriptionBase::GetEdgeVertices(const FEdgeID EdgeID, TArray<FVertexID>& OutVertexIDs) const
 {
-	if (!MeshDescription.IsEdgeValid(EdgeID))
+	if (!GetMeshDescription().IsEdgeValid(EdgeID))
 	{
 		UE_LOG(LogMeshDescription, Warning, TEXT("GetEdgeVertices: EdgeID %d doesn't exist."), EdgeID.GetValue());
 		return;
 	}
 
-	Algo::Copy(MeshDescription.GetEdgeVertices(EdgeID), OutVertexIDs);
+	Algo::Copy(GetMeshDescription().GetEdgeVertices(EdgeID), OutVertexIDs);
 }
 
 
@@ -673,51 +673,51 @@ void UMeshDescriptionBase::GetEdgeVertices(const FEdgeID EdgeID, TArray<FVertexI
 
 FPolygonID UMeshDescriptionBase::GetTrianglePolygon(FTriangleID TriangleID) const
 {
-	if (!MeshDescription.IsTriangleValid(TriangleID))
+	if (!GetMeshDescription().IsTriangleValid(TriangleID))
 	{
 		UE_LOG(LogMeshDescription, Warning, TEXT("GetTrianglePolygon: TriangleID %d doesn't exist."), TriangleID.GetValue());
 		return INDEX_NONE;
 	}
 
-	return MeshDescription.GetTrianglePolygon(TriangleID);
+	return GetMeshDescription().GetTrianglePolygon(TriangleID);
 }
 
 FPolygonGroupID UMeshDescriptionBase::GetTrianglePolygonGroup(FTriangleID TriangleID) const
 {
-	if (!MeshDescription.IsTriangleValid(TriangleID))
+	if (!GetMeshDescription().IsTriangleValid(TriangleID))
 	{
 		UE_LOG(LogMeshDescription, Warning, TEXT("GetTrianglePolygonGroup: TriangleID %d doesn't exist."), TriangleID.GetValue());
 		return INDEX_NONE;
 	}
 
-	return MeshDescription.GetTrianglePolygonGroup(TriangleID);
+	return GetMeshDescription().GetTrianglePolygonGroup(TriangleID);
 }
 
 bool UMeshDescriptionBase::IsTrianglePartOfNgon(FTriangleID TriangleID) const
 {
-	if (!MeshDescription.IsTriangleValid(TriangleID))
+	if (!GetMeshDescription().IsTriangleValid(TriangleID))
 	{
 		UE_LOG(LogMeshDescription, Warning, TEXT("IsTrianglePartOfNgon: TriangleID %d doesn't exist."), TriangleID.GetValue());
 		return false;
 	}
 
-	return MeshDescription.IsTrianglePartOfNgon(TriangleID);
+	return GetMeshDescription().IsTrianglePartOfNgon(TriangleID);
 }
 
 void UMeshDescriptionBase::GetTriangleVertexInstances(FTriangleID TriangleID, TArray<FVertexInstanceID>& OutVertexInstanceIDs) const
 {
-	if (!MeshDescription.IsTriangleValid(TriangleID))
+	if (!GetMeshDescription().IsTriangleValid(TriangleID))
 	{
 		UE_LOG(LogMeshDescription, Warning, TEXT("GetTriangleVertexInstances: TriangleID %d doesn't exist."), TriangleID.GetValue());
 		return;
 	}
 
-	Algo::Copy(MeshDescription.GetTriangleVertexInstances(TriangleID), OutVertexInstanceIDs);
+	Algo::Copy(GetMeshDescription().GetTriangleVertexInstances(TriangleID), OutVertexInstanceIDs);
 }
 
 FVertexInstanceID UMeshDescriptionBase::GetTriangleVertexInstance(FTriangleID TriangleID, int32 Index) const
 {
-	if (!MeshDescription.IsTriangleValid(TriangleID))
+	if (!GetMeshDescription().IsTriangleValid(TriangleID))
 	{
 		UE_LOG(LogMeshDescription, Warning, TEXT("GetTriangleVertexInstance: TriangleID %d doesn't exist."), TriangleID.GetValue());
 		return INDEX_NONE;
@@ -729,53 +729,53 @@ FVertexInstanceID UMeshDescriptionBase::GetTriangleVertexInstance(FTriangleID Tr
 		return INDEX_NONE;
 	}
 
-	return MeshDescription.GetTriangleVertexInstance(TriangleID, Index);
+	return GetMeshDescription().GetTriangleVertexInstance(TriangleID, Index);
 }
 
 void UMeshDescriptionBase::GetTriangleVertices(FTriangleID TriangleID, TArray<FVertexID>& OutVertexIDs) const
 {
-	if (!MeshDescription.IsTriangleValid(TriangleID))
+	if (!GetMeshDescription().IsTriangleValid(TriangleID))
 	{
 		UE_LOG(LogMeshDescription, Warning, TEXT("GetTriangleVertices: TriangleID %d doesn't exist."), TriangleID.GetValue());
 		return;
 	}
 
 	OutVertexIDs.SetNumUninitialized(3);
-	Algo::Copy(MeshDescription.GetTriangleVertices(TriangleID), OutVertexIDs);
+	Algo::Copy(GetMeshDescription().GetTriangleVertices(TriangleID), OutVertexIDs);
 }
 
 void UMeshDescriptionBase::GetTriangleEdges(FTriangleID TriangleID, TArray<FEdgeID>& OutEdgeIDs) const
 {
-	if (!MeshDescription.IsTriangleValid(TriangleID))
+	if (!GetMeshDescription().IsTriangleValid(TriangleID))
 	{
 		UE_LOG(LogMeshDescription, Warning, TEXT("GetTriangleEdges: TriangleID %d doesn't exist."), TriangleID.GetValue());
 		return;
 	}
 
 	OutEdgeIDs.SetNumUninitialized(3);
-	Algo::Copy(MeshDescription.GetTriangleEdges(TriangleID), OutEdgeIDs);
+	Algo::Copy(GetMeshDescription().GetTriangleEdges(TriangleID), OutEdgeIDs);
 }
 
 void UMeshDescriptionBase::GetTriangleAdjacentTriangles(FTriangleID TriangleID, TArray<FTriangleID>& OutTriangleIDs) const
 {
-	if (!MeshDescription.IsTriangleValid(TriangleID))
+	if (!GetMeshDescription().IsTriangleValid(TriangleID))
 	{
 		UE_LOG(LogMeshDescription, Warning, TEXT("GetTriangleAdjacentTriangles: TriangleID %d doesn't exist."), TriangleID.GetValue());
 		return;
 	}
 
-	MeshDescription.GetTriangleAdjacentTriangles(TriangleID, OutTriangleIDs);
+	GetMeshDescription().GetTriangleAdjacentTriangles(TriangleID, OutTriangleIDs);
 }
 
 FVertexInstanceID UMeshDescriptionBase::GetVertexInstanceForTriangleVertex(FTriangleID TriangleID, FVertexID VertexID) const
 {
-	if (!MeshDescription.IsTriangleValid(TriangleID))
+	if (!GetMeshDescription().IsTriangleValid(TriangleID))
 	{
 		UE_LOG(LogMeshDescription, Warning, TEXT("GetVertexInstanceForTriangleVertex: TriangleID %d doesn't exist."), TriangleID.GetValue());
 		return INDEX_NONE;
 	}
 
-	return MeshDescription.GetVertexInstanceForTriangleVertex(TriangleID, VertexID);
+	return GetMeshDescription().GetVertexInstanceForTriangleVertex(TriangleID, VertexID);
 }
 
 
@@ -784,128 +784,128 @@ FVertexInstanceID UMeshDescriptionBase::GetVertexInstanceForTriangleVertex(FTria
 
 void UMeshDescriptionBase::GetPolygonTriangles(FPolygonID PolygonID, TArray<FTriangleID>& OutTriangleIDs) const
 {
-	if (!MeshDescription.IsPolygonValid(PolygonID))
+	if (!GetMeshDescription().IsPolygonValid(PolygonID))
 	{
 		UE_LOG(LogMeshDescription, Warning, TEXT("GetPolygonTriangles: PolygonID %d doesn't exist."), PolygonID.GetValue());
 		return;
 	}
 
-	OutTriangleIDs = MeshDescription.GetPolygonTriangles(PolygonID);
+	OutTriangleIDs = GetMeshDescription().GetPolygonTriangles(PolygonID);
 }
 
 int32 UMeshDescriptionBase::GetNumPolygonTriangles(FPolygonID PolygonID) const
 {
-	if (!MeshDescription.IsPolygonValid(PolygonID))
+	if (!GetMeshDescription().IsPolygonValid(PolygonID))
 	{
 		UE_LOG(LogMeshDescription, Warning, TEXT("GetNumPolygonTriangles: PolygonID %d doesn't exist."), PolygonID.GetValue());
 		return 0;
 	}
 
-	return MeshDescription.GetNumPolygonTriangles(PolygonID);
+	return GetMeshDescription().GetNumPolygonTriangles(PolygonID);
 }
 
 void UMeshDescriptionBase::GetPolygonVertexInstances(FPolygonID PolygonID, TArray<FVertexInstanceID>& OutVertexInstanceIDs) const
 {
-	if (!MeshDescription.IsPolygonValid(PolygonID))
+	if (!GetMeshDescription().IsPolygonValid(PolygonID))
 	{
 		UE_LOG(LogMeshDescription, Warning, TEXT("GetPolygonVertexInstances: PolygonID %d doesn't exist."), PolygonID.GetValue());
 		return;
 	}
 
-	OutVertexInstanceIDs = MeshDescription.GetPolygonVertexInstances(PolygonID);
+	OutVertexInstanceIDs = GetMeshDescription().GetPolygonVertexInstances(PolygonID);
 }
 
 int32 UMeshDescriptionBase::GetNumPolygonVertices(FPolygonID PolygonID) const
 {
-	if (!MeshDescription.IsPolygonValid(PolygonID))
+	if (!GetMeshDescription().IsPolygonValid(PolygonID))
 	{
 		UE_LOG(LogMeshDescription, Warning, TEXT("GetNumPolygonVertices: PolygonID %d doesn't exist."), PolygonID.GetValue());
 		return 0;
 	}
 
-	return MeshDescription.GetNumPolygonVertices(PolygonID);
+	return GetMeshDescription().GetNumPolygonVertices(PolygonID);
 }
 
 void UMeshDescriptionBase::GetPolygonVertices(FPolygonID PolygonID, TArray<FVertexID>& OutVertexIDs) const
 {
-	if (!MeshDescription.IsPolygonValid(PolygonID))
+	if (!GetMeshDescription().IsPolygonValid(PolygonID))
 	{
 		UE_LOG(LogMeshDescription, Warning, TEXT("GetPolygonVertices: PolygonID %d doesn't exist."), PolygonID.GetValue());
 		return;
 	}
 
-	MeshDescription.GetPolygonVertices(PolygonID, OutVertexIDs);
+	GetMeshDescription().GetPolygonVertices(PolygonID, OutVertexIDs);
 }
 
 void UMeshDescriptionBase::GetPolygonPerimeterEdges(FPolygonID PolygonID, TArray<FEdgeID>& OutEdgeIDs) const
 {
-	if (!MeshDescription.IsPolygonValid(PolygonID))
+	if (!GetMeshDescription().IsPolygonValid(PolygonID))
 	{
 		UE_LOG(LogMeshDescription, Warning, TEXT("GetPolygonPerimeterEdges: PolygonID %d doesn't exist."), PolygonID.GetValue());
 		return;
 	}
 
-	MeshDescription.GetPolygonPerimeterEdges(PolygonID, OutEdgeIDs);
+	GetMeshDescription().GetPolygonPerimeterEdges(PolygonID, OutEdgeIDs);
 }
 
 void UMeshDescriptionBase::GetPolygonInternalEdges(FPolygonID PolygonID, TArray<FEdgeID>& OutEdgeIDs) const
 {
-	if (!MeshDescription.IsPolygonValid(PolygonID))
+	if (!GetMeshDescription().IsPolygonValid(PolygonID))
 	{
 		UE_LOG(LogMeshDescription, Warning, TEXT("GetPolygonInternalEdges: PolygonID %d doesn't exist."), PolygonID.GetValue());
 		return;
 	}
 
-	MeshDescription.GetPolygonInternalEdges(PolygonID, OutEdgeIDs);
+	GetMeshDescription().GetPolygonInternalEdges(PolygonID, OutEdgeIDs);
 }
 
 int32 UMeshDescriptionBase::GetNumPolygonInternalEdges(FPolygonID PolygonID) const
 {
-	if (!MeshDescription.IsPolygonValid(PolygonID))
+	if (!GetMeshDescription().IsPolygonValid(PolygonID))
 	{
 		UE_LOG(LogMeshDescription, Warning, TEXT("GetNumPolygonInternalEdges: PolygonID %d doesn't exist."), PolygonID.GetValue());
 		return 0;
 	}
 
-	return MeshDescription.GetNumPolygonInternalEdges(PolygonID);
+	return GetMeshDescription().GetNumPolygonInternalEdges(PolygonID);
 }
 
 void UMeshDescriptionBase::GetPolygonAdjacentPolygons(FPolygonID PolygonID, TArray<FPolygonID>& OutPolygonIDs) const
 {
-	if (!MeshDescription.IsPolygonValid(PolygonID))
+	if (!GetMeshDescription().IsPolygonValid(PolygonID))
 	{
 		UE_LOG(LogMeshDescription, Warning, TEXT("GetPolygonAdjacentPolygons: PolygonID %d doesn't exist."), PolygonID.GetValue());
 		return;
 	}
 
-	MeshDescription.GetPolygonAdjacentPolygons(PolygonID, OutPolygonIDs);
+	GetMeshDescription().GetPolygonAdjacentPolygons(PolygonID, OutPolygonIDs);
 }
 
 FPolygonGroupID UMeshDescriptionBase::GetPolygonPolygonGroup(FPolygonID PolygonID) const
 {
-	if (!MeshDescription.IsPolygonValid(PolygonID))
+	if (!GetMeshDescription().IsPolygonValid(PolygonID))
 	{
 		UE_LOG(LogMeshDescription, Warning, TEXT("GetPolygonPolygonGroup: PolygonID %d doesn't exist."), PolygonID.GetValue());
 		return INDEX_NONE;
 	}
 
-	return MeshDescription.GetPolygonPolygonGroup(PolygonID);
+	return GetMeshDescription().GetPolygonPolygonGroup(PolygonID);
 }
 
 FVertexInstanceID UMeshDescriptionBase::GetVertexInstanceForPolygonVertex(FPolygonID PolygonID, FVertexID VertexID) const
 {
-	if (!MeshDescription.IsPolygonValid(PolygonID))
+	if (!GetMeshDescription().IsPolygonValid(PolygonID))
 	{
 		UE_LOG(LogMeshDescription, Warning, TEXT("GetVertexInstanceForPolygonVertex: PolygonID %d doesn't exist."), PolygonID.GetValue());
 		return INDEX_NONE;
 	}
 
-	return MeshDescription.GetVertexInstanceForPolygonVertex(PolygonID, VertexID);
+	return GetMeshDescription().GetVertexInstanceForPolygonVertex(PolygonID, VertexID);
 }
 
 void UMeshDescriptionBase::SetPolygonVertexInstances(FPolygonID PolygonID, const TArray<FVertexInstanceID>& VertexInstanceIDs)
 {
-	if (!MeshDescription.IsPolygonValid(PolygonID))
+	if (!GetMeshDescription().IsPolygonValid(PolygonID))
 	{
 		UE_LOG(LogMeshDescription, Warning, TEXT("SetPolygonVertexInstances: PolygonID %d doesn't exist."), PolygonID.GetValue());
 		return;
@@ -913,53 +913,53 @@ void UMeshDescriptionBase::SetPolygonVertexInstances(FPolygonID PolygonID, const
 
 	for (FVertexInstanceID VertexInstanceID : VertexInstanceIDs)
 	{
-		if (!MeshDescription.IsVertexInstanceValid(VertexInstanceID))
+		if (!GetMeshDescription().IsVertexInstanceValid(VertexInstanceID))
 		{
 			UE_LOG(LogMeshDescription, Warning, TEXT("SetPolygonVertexInstances: VertexInstanceID %d doesn't exist."), VertexInstanceID.GetValue());
 			return;
 		}
 	}
 
-	MeshDescription.SetPolygonVertexInstances(PolygonID, VertexInstanceIDs);
+	GetMeshDescription().SetPolygonVertexInstances(PolygonID, VertexInstanceIDs);
 }
 
 void UMeshDescriptionBase::SetPolygonPolygonGroup(FPolygonID PolygonID, FPolygonGroupID PolygonGroupID)
 {
-	if (!MeshDescription.IsPolygonValid(PolygonID))
+	if (!GetMeshDescription().IsPolygonValid(PolygonID))
 	{
 		UE_LOG(LogMeshDescription, Warning, TEXT("SetPolygonPolygonGroup: PolygonID %d doesn't exist."), PolygonID.GetValue());
 		return;
 	}
 
-	if (!MeshDescription.IsPolygonGroupValid(PolygonGroupID))
+	if (!GetMeshDescription().IsPolygonGroupValid(PolygonGroupID))
 	{
 		UE_LOG(LogMeshDescription, Warning, TEXT("SetPolygonPolygonGroup: PolygonGroupID %d doesn't exist."), PolygonGroupID.GetValue());
 		return;
 	}
 
-	MeshDescription.SetPolygonPolygonGroup(PolygonID, PolygonGroupID);
+	GetMeshDescription().SetPolygonPolygonGroup(PolygonID, PolygonGroupID);
 }
 
 void UMeshDescriptionBase::ReversePolygonFacing(FPolygonID PolygonID)
 {
-	if (!MeshDescription.IsPolygonValid(PolygonID))
+	if (!GetMeshDescription().IsPolygonValid(PolygonID))
 	{
 		UE_LOG(LogMeshDescription, Warning, TEXT("ReversePolygonFacing: PolygonID %d doesn't exist."), PolygonID.GetValue());
 		return;
 	}
 
-	MeshDescription.ReversePolygonFacing(PolygonID);
+	GetMeshDescription().ReversePolygonFacing(PolygonID);
 }
 
 void UMeshDescriptionBase::ComputePolygonTriangulation(FPolygonID PolygonID)
 {
-	if (!MeshDescription.IsPolygonValid(PolygonID))
+	if (!GetMeshDescription().IsPolygonValid(PolygonID))
 	{
 		UE_LOG(LogMeshDescription, Warning, TEXT("ComputePolygonTriangulation: PolygonID %d doesn't exist."), PolygonID.GetValue());
 		return;
 	}
 
-	MeshDescription.ComputePolygonTriangulation(PolygonID);
+	GetMeshDescription().ComputePolygonTriangulation(PolygonID);
 }
 
 
@@ -968,22 +968,22 @@ void UMeshDescriptionBase::ComputePolygonTriangulation(FPolygonID PolygonID)
 
 void UMeshDescriptionBase::GetPolygonGroupPolygons(FPolygonGroupID PolygonGroupID, TArray<FPolygonID>& OutPolygonIDs) const
 {
-	if (!MeshDescription.IsPolygonGroupValid(PolygonGroupID))
+	if (!GetMeshDescription().IsPolygonGroupValid(PolygonGroupID))
 	{
 		UE_LOG(LogMeshDescription, Warning, TEXT("GetPolygonGroupPolygons: PolygonGroupID %d doesn't exist."), PolygonGroupID.GetValue());
 		return;
 	}
 
-	OutPolygonIDs = MeshDescription.GetPolygonGroupPolygonIDs(PolygonGroupID);
+	OutPolygonIDs = GetMeshDescription().GetPolygonGroupPolygonIDs(PolygonGroupID);
 }
 
 int32 UMeshDescriptionBase::GetNumPolygonGroupPolygons(FPolygonGroupID PolygonGroupID) const
 {
-	if (!MeshDescription.IsPolygonGroupValid(PolygonGroupID))
+	if (!GetMeshDescription().IsPolygonGroupValid(PolygonGroupID))
 	{
 		UE_LOG(LogMeshDescription, Warning, TEXT("GetNumPolygonGroupPolygons: PolygonGroupID %d doesn't exist."), PolygonGroupID.GetValue());
 		return 0;
 	}
 
-	return MeshDescription.GetNumPolygonGroupPolygons(PolygonGroupID);
+	return GetMeshDescription().GetNumPolygonGroupPolygons(PolygonGroupID);
 }
diff --git a/Engine/Source/Runtime/MeshDescription/Public/MeshDescriptionBase.h b/Engine/Source/Runtime/MeshDescription/Public/MeshDescriptionBase.h
--- a/Engine/Source/Runtime/MeshDescription/Public/MeshDescriptionBase.h
+++ b/Engine/Source/Runtime/MeshDescription/Public/MeshDescriptionBase.h
@@ -19,13 +19,21 @@ public:
 	//
 
 	/** Get a reference to the actual mesh description */
-	FMeshDescription& GetMeshDescription() { return MeshDescription; }
-	const FMeshDescription& GetMeshDescription() const { return MeshDescription; }
+	FMeshDescription& GetMeshDescription() { return MeshDescriptionPtr ? OwnedMeshDescription : *MeshDescriptionPtr; }
+	const FMeshDescription& GetMeshDescription() const { return MeshDescriptionPtr ? OwnedMeshDescription : *MeshDescriptionPtr; }
 
 	/** Set the mesh description */
 	void SetMeshDescription(FMeshDescription InMeshDescription)
 	{ 
-		MeshDescription = MoveTemp(InMeshDescription);
+		OwnedMeshDescription = MoveTemp(InMeshDescription);
+		MeshDescriptionPtr = nullptr;
+	}
+
+	/** Sets a reference to another mesh description */
+	void SetMeshDescriptionRef(FMeshDescription* InMeshDescriptionRef)
+	{
+		MeshDescriptionPtr = InMeshDescriptionRef;
+		OwnedMeshDescription.Empty();
 	}
 
 	// UMeshDescriptionBase interface
@@ -38,42 +46,42 @@ public:
 	void Reset();
 
 	/** Accessors for mesh element arrays */
-	FVertexArray& Vertices() { return MeshDescription.Vertices(); }
-	const FVertexArray& Vertices() const { return MeshDescription.Vertices(); }
+	FVertexArray& Vertices() { return GetMeshDescription().Vertices(); }
+	const FVertexArray& Vertices() const { return GetMeshDescription().Vertices(); }
 
-	FVertexInstanceArray& VertexInstances() { return MeshDescription.VertexInstances(); }
-	const FVertexInstanceArray& VertexInstances() const { return MeshDescription.VertexInstances(); }
+	FVertexInstanceArray& VertexInstances() { return GetMeshDescription().VertexInstances(); }
+	const FVertexInstanceArray& VertexInstances() const { return GetMeshDescription().VertexInstances(); }
 
-	FEdgeArray& Edges() { return MeshDescription.Edges(); }
-	const FEdgeArray& Edges() const { return MeshDescription.Edges(); }
+	FEdgeArray& Edges() { return GetMeshDescription().Edges(); }
+	const FEdgeArray& Edges() const { return GetMeshDescription().Edges(); }
 
-	FTriangleArray& Triangles() { return MeshDescription.Triangles(); }
-	const FTriangleArray& Triangles() const { return MeshDescription.Triangles(); }
+	FTriangleArray& Triangles() { return GetMeshDescription().Triangles(); }
+	const FTriangleArray& Triangles() const { return GetMeshDescription().Triangles(); }
 
-	FPolygonArray& Polygons() { return MeshDescription.Polygons(); }
-	const FPolygonArray& Polygons() const { return MeshDescription.Polygons(); }
+	FPolygonArray& Polygons() { return GetMeshDescription().Polygons(); }
+	const FPolygonArray& Polygons() const { return GetMeshDescription().Polygons(); }
 
-	FPolygonGroupArray& PolygonGroups() { return MeshDescription.PolygonGroups(); }
-	const FPolygonGroupArray& PolygonGroups() const { return MeshDescription.PolygonGroups(); }
+	FPolygonGroupArray& PolygonGroups() { return GetMeshDescription().PolygonGroups(); }
+	const FPolygonGroupArray& PolygonGroups() const { return GetMeshDescription().PolygonGroups(); }
 
 	/** Accessors for mesh element attributes */
-	TAttributesSet<FVertexID>& VertexAttributes() { return MeshDescription.VertexAttributes(); }
-	const TAttributesSet<FVertexID>& VertexAttributes() const { return MeshDescription.VertexAttributes(); }
+	TAttributesSet<FVertexID>& VertexAttributes() { return GetMeshDescription().VertexAttributes(); }
+	const TAttributesSet<FVertexID>& VertexAttributes() const { return GetMeshDescription().VertexAttributes(); }
 
-	TAttributesSet<FVertexInstanceID>& VertexInstanceAttributes() { return MeshDescription.VertexInstanceAttributes(); }
-	const TAttributesSet<FVertexInstanceID>& VertexInstanceAttributes() const { return MeshDescription.VertexInstanceAttributes(); }
+	TAttributesSet<FVertexInstanceID>& VertexInstanceAttributes() { return GetMeshDescription().VertexInstanceAttributes(); }
+	const TAttributesSet<FVertexInstanceID>& VertexInstanceAttributes() const { return GetMeshDescription().VertexInstanceAttributes(); }
 
-	TAttributesSet<FEdgeID>& EdgeAttributes() { return MeshDescription.EdgeAttributes(); }
-	const TAttributesSet<FEdgeID>& EdgeAttributes() const { return MeshDescription.EdgeAttributes(); }
+	TAttributesSet<FEdgeID>& EdgeAttributes() { return GetMeshDescription().EdgeAttributes(); }
+	const TAttributesSet<FEdgeID>& EdgeAttributes() const { return GetMeshDescription().EdgeAttributes(); }
 
-	TAttributesSet<FTriangleID>& TriangleAttributes() { return MeshDescription.TriangleAttributes(); }
-	const TAttributesSet<FTriangleID>& TriangleAttributes() const { return MeshDescription.TriangleAttributes(); }
+	TAttributesSet<FTriangleID>& TriangleAttributes() { return GetMeshDescription().TriangleAttributes(); }
+	const TAttributesSet<FTriangleID>& TriangleAttributes() const { return GetMeshDescription().TriangleAttributes(); }
 
-	TAttributesSet<FPolygonID>& PolygonAttributes() { return MeshDescription.PolygonAttributes(); }
-	const TAttributesSet<FPolygonID>& PolygonAttributes() const { return MeshDescription.PolygonAttributes(); }
+	TAttributesSet<FPolygonID>& PolygonAttributes() { return GetMeshDescription().PolygonAttributes(); }
+	const TAttributesSet<FPolygonID>& PolygonAttributes() const { return GetMeshDescription().PolygonAttributes(); }
 
-	TAttributesSet<FPolygonGroupID>& PolygonGroupAttributes() { return MeshDescription.PolygonGroupAttributes(); }
-	const TAttributesSet<FPolygonGroupID>& PolygonGroupAttributes() const { return MeshDescription.PolygonGroupAttributes(); }
+	TAttributesSet<FPolygonGroupID>& PolygonGroupAttributes() { return GetMeshDescription().PolygonGroupAttributes(); }
+	const TAttributesSet<FPolygonGroupID>& PolygonGroupAttributes() const { return GetMeshDescription().PolygonGroupAttributes(); }
 
 	/** Accessors for cached vertex position array */
 	TVertexAttributesRef<FVector> GetVertexPositions() { return GetRequiredAttributes().GetVertexPositions(); }
@@ -91,7 +99,7 @@ public:
 
 
 	///////////////////////////////////////////////////////
-	// Create / remove mesh elements
+	// Create / remove / count mesh elements
 
 	/** Reserves space for this number of new vertices */
 	UFUNCTION(BlueprintCallable, Category="MeshDescription")
@@ -113,6 +121,10 @@ public:
 	UFUNCTION(BlueprintPure, Category="MeshDescription")
 	bool IsVertexValid(FVertexID VertexID) const;
 
+	/** Returns the number of vertices */
+	UFUNCTION(BlueprintPure, Category="MeshDescription")
+	int32 GetVertexCount() const { return Vertices().Num(); }
+
 	/** Reserves space for this number of new vertex instances */
 	UFUNCTION(BlueprintCallable, Category="MeshDescription")
 	void ReserveNewVertexInstances(int32 NumberOfNewVertexInstances);
@@ -133,6 +145,10 @@ public:
 	UFUNCTION(BlueprintPure, Category="MeshDescription")
 	bool IsVertexInstanceValid(FVertexInstanceID VertexInstanceID) const;
 
+	/** Returns the number of vertex instances */
+	UFUNCTION(BlueprintPure, Category="MeshDescription")
+	int32 GetVertexInstanceCount() const { return VertexInstances().Num(); }
+
 	/** Reserves space for this number of new edges */
 	UFUNCTION(BlueprintCallable, Category="MeshDescription")
 	void ReserveNewEdges(int32 NumberOfNewEdges);
@@ -153,6 +169,10 @@ public:
 	UFUNCTION(BlueprintPure, Category="MeshDescription")
 	bool IsEdgeValid(FEdgeID EdgeID) const;
 
+	/** Returns the number of edges */
+	UFUNCTION(BlueprintPure, Category="MeshDescription")
+	int32 GetEdgeCount() const { return Edges().Num(); }
+
 	/** Reserves space for this number of new triangles */
 	UFUNCTION(BlueprintCallable, Category="MeshDescription")
 	void ReserveNewTriangles(int32 NumberOfNewTriangles);
@@ -173,6 +193,10 @@ public:
 	UFUNCTION(BlueprintPure, Category="MeshDescription")
 	bool IsTriangleValid(const FTriangleID TriangleID) const;
 
+	/** Returns the number of triangles */
+	UFUNCTION(BlueprintPure, Category="MeshDescription")
+	int32 GetTriangleCount() const { return Triangles().Num(); }
+
 	/** Reserves space for this number of new polygons */
 	UFUNCTION(BlueprintCallable, Category="MeshDescription")
 	void ReserveNewPolygons(const int32 NumberOfNewPolygons);
@@ -193,6 +217,10 @@ public:
 	UFUNCTION(BlueprintPure, Category="MeshDescription")
 	bool IsPolygonValid(FPolygonID PolygonID) const;
 
+	/** Returns the number of polygons */
+	UFUNCTION(BlueprintPure, Category="MeshDescription")
+	int32 GetPolygonCount() const { return Polygons().Num(); }
+
 	/** Reserves space for this number of new polygon groups */
 	UFUNCTION(BlueprintCallable, Category="MeshDescription")
 	void ReserveNewPolygonGroups(int32 NumberOfNewPolygonGroups);
@@ -213,6 +241,10 @@ public:
 	UFUNCTION(BlueprintPure, Category="MeshDescription")
 	bool IsPolygonGroupValid(FPolygonGroupID PolygonGroupID) const;
 
+	/** Returns the number of polygon groups */
+	UFUNCTION(BlueprintPure, Category="MeshDescription")
+	int32 GetPolygonGroupCount() const { return PolygonGroups().Num(); }
+
 
 	//////////////////////////////////////////////////////////////////////
 	// Vertex operations
@@ -451,6 +483,7 @@ public:
 	int32 GetNumPolygonGroupPolygons(FPolygonGroupID PolygonGroupID) const;
 
 protected:
-	FMeshDescription MeshDescription;
+	FMeshDescription OwnedMeshDescription;
+	FMeshDescription* MeshDescriptionPtr = nullptr;
 	TUniquePtr<FMeshAttributes> RequiredAttributes;
 };
diff --git a/Engine/Source/Runtime/StaticMeshDescription/Private/StaticMeshDescription.cpp b/Engine/Source/Runtime/StaticMeshDescription/Private/StaticMeshDescription.cpp
--- a/Engine/Source/Runtime/StaticMeshDescription/Private/StaticMeshDescription.cpp
+++ b/Engine/Source/Runtime/StaticMeshDescription/Private/StaticMeshDescription.cpp
@@ -7,62 +7,62 @@
 
 void UStaticMeshDescription::RegisterAttributes()
 {
-	RequiredAttributes = MakeUnique<FStaticMeshAttributes>(MeshDescription);
+	RequiredAttributes = MakeUnique<FStaticMeshAttributes>(GetMeshDescription());
 	RequiredAttributes->Register();
 }
 
 
 FVector2D UStaticMeshDescription::GetVertexInstanceUV(FVertexInstanceID VertexInstanceID, int32 UVIndex) const
 {
-	if (!MeshDescription.IsVertexInstanceValid(VertexInstanceID))
+	if (!GetMeshDescription().IsVertexInstanceValid(VertexInstanceID))
 	{
 		UE_LOG(LogMeshDescription, Warning, TEXT("GetVertexInstanceUV: VertexInstanceID %d doesn't exist."), VertexInstanceID.GetValue());
 		return FVector2D::ZeroVector;
 	}
 
-	if (!MeshDescription.VertexInstanceAttributes().HasAttribute(MeshAttribute::VertexInstance::TextureCoordinate))
+	if (!GetMeshDescription().VertexInstanceAttributes().HasAttribute(MeshAttribute::VertexInstance::TextureCoordinate))
 	{
 		UE_LOG(LogMeshDescription, Warning, TEXT("GetVertexInstanceUV: VertexInstanceAttribute TextureCoordinate doesn't exist."));
 		return FVector2D::ZeroVector;
 	}
 
-	return MeshDescription.VertexInstanceAttributes().GetAttribute<FVector2D>(VertexInstanceID, MeshAttribute::VertexInstance::TextureCoordinate, UVIndex);
+	return GetMeshDescription().VertexInstanceAttributes().GetAttribute<FVector2D>(VertexInstanceID, MeshAttribute::VertexInstance::TextureCoordinate, UVIndex);
 }
 
 
 void UStaticMeshDescription::SetVertexInstanceUV(FVertexInstanceID VertexInstanceID, FVector2D UV, int32 UVIndex)
 {
-	if (!MeshDescription.IsVertexInstanceValid(VertexInstanceID))
+	if (!GetMeshDescription().IsVertexInstanceValid(VertexInstanceID))
 	{
 		UE_LOG(LogMeshDescription, Warning, TEXT("SetVertexInstanceUV: VertexInstanceID %d doesn't exist."), VertexInstanceID.GetValue());
 		return;
 	}
 
-	if (!MeshDescription.VertexInstanceAttributes().HasAttribute(MeshAttribute::VertexInstance::TextureCoordinate))
+	if (!GetMeshDescription().VertexInstanceAttributes().HasAttribute(MeshAttribute::VertexInstance::TextureCoordinate))
 	{
 		UE_LOG(LogMeshDescription, Warning, TEXT("SetVertexInstanceUV: VertexInstanceAttribute TextureCoordinate doesn't exist."));
 		return;
 	}
 
-	MeshDescription.VertexInstanceAttributes().SetAttribute(VertexInstanceID, MeshAttribute::VertexInstance::TextureCoordinate, UVIndex, UV);
+	GetMeshDescription().VertexInstanceAttributes().SetAttribute(VertexInstanceID, MeshAttribute::VertexInstance::TextureCoordinate, UVIndex, UV);
 }
 
 
 void UStaticMeshDescription::SetPolygonGroupMaterialSlotName(FPolygonGroupID PolygonGroupID, const FName& SlotName)
 {
-	if (!MeshDescription.IsPolygonGroupValid(PolygonGroupID))
+	if (!GetMeshDescription().IsPolygonGroupValid(PolygonGroupID))
 	{
 		UE_LOG(LogMeshDescription, Warning, TEXT("SetPolygonGroupMaterialSlotName: PolygonGroupID %d doesn't exist."), PolygonGroupID.GetValue());
 		return;
 	}
 
-	if (!MeshDescription.PolygonGroupAttributes().HasAttribute(MeshAttribute::PolygonGroup::ImportedMaterialSlotName))
+	if (!GetMeshDescription().PolygonGroupAttributes().HasAttribute(MeshAttribute::PolygonGroup::ImportedMaterialSlotName))
 	{
 		UE_LOG(LogMeshDescription, Warning, TEXT("SetPolygonGroupMaterialSlotName: PolygonGroupAttribute ImportedMaterialSlotName doesn't exist."));
 		return;
 	}
 
-	MeshDescription.PolygonGroupAttributes().SetAttribute(PolygonGroupID, MeshAttribute::PolygonGroup::ImportedMaterialSlotName, 0, SlotName);
+	GetMeshDescription().PolygonGroupAttributes().SetAttribute(PolygonGroupID, MeshAttribute::PolygonGroup::ImportedMaterialSlotName, 0, SlotName);
 }
 
 
@@ -74,6 +74,7 @@ void UStaticMeshDescription::CreateCube(FVector Center, FVector HalfExtents, FPo
 										FPolygonID& PolygonID_PlusZ,
 										FPolygonID& PolygonID_MinusZ)
 {
+	FMeshDescription& MeshDescription = GetMeshDescription();
 	TVertexAttributesRef<FVector> Positions = GetVertexPositions();
 
 	FVertexID VertexIDs[8];
@@ -93,7 +94,7 @@ void UStaticMeshDescription::CreateCube(FVector Center, FVector HalfExtents, FPo
 	Positions[VertexIDs[6]] = Center + HalfExtents * FVector( 1.0f, -1.0f, -1.0f);
 	Positions[VertexIDs[7]] = Center + HalfExtents * FVector( 1.0f,  1.0f, -1.0f);
 
-	auto MakePolygon = [this, &VertexIDs, PolygonGroup](int32 P0, int32 P1, int32 P2, int32 P3) -> FPolygonID
+	auto MakePolygon = [this, &MeshDescription, &VertexIDs, PolygonGroup](int32 P0, int32 P1, int32 P2, int32 P3) -> FPolygonID
 	{
 		FVertexInstanceID VertexInstanceIDs[4];
 		VertexInstanceIDs[0] = MeshDescription.CreateVertexInstance(VertexIDs[P0]);
