Choose the owned mesh description if the mesh description reference is invalid.

diff --git a/Engine/Source/Runtime/MeshDescription/Public/MeshDescriptionBase.h b/Engine/Source/Runtime/MeshDescription/Public/MeshDescriptionBase.h
--- a/Engine/Source/Runtime/MeshDescription/Public/MeshDescriptionBase.h
+++ b/Engine/Source/Runtime/MeshDescription/Public/MeshDescriptionBase.h
@@ -19,8 +19,29 @@ public:
 	//
 
 	/** Get a reference to the actual mesh description */
-	FMeshDescription& GetMeshDescription() { return MeshDescriptionPtr ? OwnedMeshDescription : *MeshDescriptionPtr; }
-	const FMeshDescription& GetMeshDescription() const { return MeshDescriptionPtr ? OwnedMeshDescription : *MeshDescriptionPtr; }
+	FMeshDescription& GetMeshDescription()
+	{
+		if (MeshDescriptionPtr)
+		{
+			return *MeshDescriptionPtr;
+		}
+		else
+		{
+			return OwnedMeshDescription;
+		}
+	}
+
+	const FMeshDescription& GetMeshDescription() const
+	{
+		if (MeshDescriptionPtr)
+		{
+			return *MeshDescriptionPtr;
+		}
+		else
+		{
+			return OwnedMeshDescription;
+		}
+	}
 
 	/** Set the mesh description */
 	void SetMeshDescription(FMeshDescription InMeshDescription)
