MeshAttributeArray improvements:
- allowed TMeshAttributeRef<TArrayView<T>> to be initialized from *any* attribute type, thus providing a generic way of accessing any type of attribute, be it simple, compound or unbounded array.

ArrayAttribute improvements:
- sped up inserting by not requiring every element of a chunk to be valid if it has not yet been used
- added new API to provide Reserve / Add / AddUnique etc

Indexer improvements:
- fixed issue allowing an attribute with a compound index attribute (e.g. int[3]) to be indexed multiple times

diff --git a/Engine/Plugins/Tests/EditorTests/Source/EditorTests/Private/UnrealEd/MeshDescriptionTests.cpp b/Engine/Plugins/Tests/EditorTests/Source/EditorTests/Private/UnrealEd/MeshDescriptionTests.cpp
--- a/Engine/Plugins/Tests/EditorTests/Source/EditorTests/Private/UnrealEd/MeshDescriptionTests.cpp
+++ b/Engine/Plugins/Tests/EditorTests/Source/EditorTests/Private/UnrealEd/MeshDescriptionTests.cpp
@@ -946,3 +946,83 @@ bool FMeshDescriptionBuilderTest::RunTest(const FString& Parameters)
 
 	return true;
 }
+
+
+IMPLEMENT_SIMPLE_AUTOMATION_TEST(FMeshDescriptionArrayAttributeTest, "Editor.Meshes.MeshDescription.ArrayAttribute", EAutomationTestFlags::EditorContext | EAutomationTestFlags::EngineFilter);
+
+bool FMeshDescriptionArrayAttributeTest::RunTest(const FString& Parameters)
+{
+	FAttributesSetBase AttributesSet;
+	AttributesSet.RegisterAttribute<int[]>("Test");
+
+	TMeshAttributesArray<TArrayAttribute<int>> Attributes = AttributesSet.GetAttributesRef<TArrayAttribute<int>>("Test");
+	TestTrue("AttributeRef valid", Attributes.IsValid());
+	TestEqual("AttributeRef element count", Attributes.GetNumElements(), 0);
+	TestEqual("AttributeRef channel count", Attributes.GetNumChannels(), 1);
+
+	AttributesSet.Initialize(10);
+	TestEqual("AttributeRef element count", Attributes.GetNumElements(), 10);
+
+	TArrayAttribute<int> Element0 = Attributes.Get(0);
+	TestEqual("Default element array size", Element0.Num(), 0);
+	Element0.Add(42);
+	TestEqual("New element value", Element0[0], 42);
+	Element0.Add(43);
+	TestEqual("New element value", Element0[1], 43);
+	TestEqual("New element array size", Element0.Num(), 2);
+
+	TArrayAttribute<int> Element1 = Attributes.Get(1);
+	Element1.Add(142);
+	Element1.Add(143);
+	Element1.Add(144);
+	TestEqual("Second element array size", Element1.Num(), 3);
+
+	Element0.Add(44);
+	Element0.Add(44);
+	Element0.Add(46);
+	TestEqual("Check after insert 0", Element0[1], 43);
+	TestEqual("Check after insert 1", Element0[2], 44);
+	TestEqual("Check after insert 2", Element0[3], 44);
+	TestEqual("Check after insert 3", Element0[4], 46);
+	TestEqual("Check after insert 4", Element1[0], 142);
+
+	Element0.InsertDefaulted(1, 3);
+	TestEqual("Check after insert 5", Element0[0], 42);
+	TestEqual("Check after insert 6", Element0[1], 0);
+	TestEqual("Check after insert 7", Element0[4], 43);
+	TestEqual("Check after insert 8", Element1[0], 142);
+	TestEqual("Check after insert 9", Element0.Num(), 8);
+
+	Element0.Reserve(20);
+	TestEqual("Check after reserve", Element1[0], 142);
+
+	AttributesSet.SetNumElements(255);
+	TArrayAttribute<int> Element254 = Attributes.Get(254);
+	Element254.Add(25400);
+
+	TestEqual("Add element 254", Element254[0], 25400);
+	TestEqual("Check element 200 empty", Attributes.Get(200).Num(), 0);
+
+	AttributesSet.SetNumElements(256);
+	TArrayAttribute<int> Element255 = Attributes.Get(255);
+	Element255.Add(25500);
+	TestEqual("Add element 255 1", Element254[0], 25400);
+	TestEqual("Add element 255 2", Element255[0], 25500);
+
+	AttributesSet.SetNumElements(300);
+	TArrayAttribute<int> Element290 = Attributes.Get(290);
+	Element290.Add(29000);
+	TestEqual("Check element 256 empty", Attributes.Get(256).Num(), 0);
+	TestEqual("Test element 290", Element290[0], 29000);
+
+	AttributesSet.SetNumElements(522);
+	TArrayAttribute<int> Element514 = Attributes.Get(514);
+	Element514.Add(51400);
+	TestEqual("Check element 512 empty", Attributes.Get(512).Num(), 0);
+	TestEqual("Test element 514", Element514[0], 51400);
+
+	Element0.Remove(44);
+	TestEqual("Remove", Element0.Num(), 6);
+
+	return true;
+}
diff --git a/Engine/Source/Runtime/MeshDescription/Public/AttributeArrayContainer.h b/Engine/Source/Runtime/MeshDescription/Public/AttributeArrayContainer.h
--- a/Engine/Source/Runtime/MeshDescription/Public/AttributeArrayContainer.h
+++ b/Engine/Source/Runtime/MeshDescription/Public/AttributeArrayContainer.h
@@ -25,9 +25,16 @@ public:
 	{
 		// For unbounded arrays, the default value is an empty subarray.
 		check(ElementCount >= 0);
-		int32 NumChunks = (ElementCount + ChunkSize - 1) >> ChunkBits;
+		int32 NumChunks = GetNumRequiredChunks(ElementCount);
 		Chunks.Empty(NumChunks);
 		Chunks.SetNum(NumChunks);
+		if (NumChunks > 0)
+		{
+			// Set num elements in the final chunk to the right value.
+			// All previous chunks are default initialized to be 'full'.
+			Chunks.Last().NumElements = GetElementsInLastChunk(ElementCount);
+		}
+
 		NumElements = ElementCount;
 	}
 
@@ -36,30 +43,53 @@ public:
 	{
 		check(ElementCount >= 0);
 
-		int32 NumChunks = (ElementCount + ChunkSize - 1) >> ChunkBits;
+		int32 OldNumChunks = Chunks.Num();
+		int32 NumChunks = GetNumRequiredChunks(ElementCount);
 		Chunks.SetNum(NumChunks);
 
 		if (ElementCount < NumElements)
 		{
 			// Case where we're shrinking the unbounded array
-			int32 IndexInLastChunk = ElementCount & ChunkMask;
-			if (IndexInLastChunk > 0)
+			if (NumChunks > 0)
 			{
-				FChunk& LastChunk = Chunks.Last();
-				int32 LastIndex = LastChunk.StartIndex[IndexInLastChunk - 1] + LastChunk.MaxCount[IndexInLastChunk - 1];
-				LastChunk.Data.SetNum(LastIndex);
-				for (; IndexInLastChunk < ChunkSize; ++IndexInLastChunk)
-				{
-					LastChunk.StartIndex[IndexInLastChunk] = LastIndex;
-					LastChunk.Count[IndexInLastChunk] = 0;
-					LastChunk.MaxCount[IndexInLastChunk] = 0;
-				}
+				Chunks.Last().NumElements = GetElementsInLastChunk(ElementCount);
 			}
 		}
-		else
+		else if (ElementCount > NumElements)
 		{
-			// If we're growing the unbounded array, there's nothing to do;
-			// the excess indices in the old last chunk will already be set up with zero length, pointing at the off-the-end element of the data array
+			// Case where we're growing the unbounded array
+			int32 IndexInOldLastChunk = NumElements & ChunkMask;
+			int32 ElementsInLastChunk = GetElementsInLastChunk(ElementCount);
+
+			if (IndexInOldLastChunk > 0)
+			{
+				// If the original final chunk had unused elements, we need to initialize their data pointers up to either:
+				// - the new final element, if it is still the final chunk
+				// - the end (ChunkSize) if it is no longer the final chunk
+				int32 LastIndex = (NumChunks == OldNumChunks) ? ElementsInLastChunk : ChunkSize;
+
+				FChunk& OldLastChunk = Chunks[OldNumChunks - 1];
+
+				// Determine the start index in the data block of the next element,
+				// and resize it to that amount
+				int32 DataStartIndex = OldLastChunk.StartIndex[IndexInOldLastChunk - 1] + OldLastChunk.MaxCount[IndexInOldLastChunk - 1];
+				OldLastChunk.Data.SetNum(DataStartIndex, false);
+
+				for (; IndexInOldLastChunk < LastIndex; IndexInOldLastChunk++)
+				{
+					OldLastChunk.StartIndex[IndexInOldLastChunk] = DataStartIndex;
+					OldLastChunk.Count[IndexInOldLastChunk] = 0;
+					OldLastChunk.MaxCount[IndexInOldLastChunk] = 0;
+				}
+
+				// Update the element count for the old last chunk
+				OldLastChunk.NumElements = LastIndex;
+			}
+
+			// Any further chunks which have just been initialized will default to being 'full', i.e. with ChunkSize elements
+			// The very last chunk will need initializing to the correct value.
+			// If this happens to coincide with the old last chunk, it doesn't hurt to do this again here.
+			Chunks.Last().NumElements = ElementsInLastChunk;
 		}
 
 		NumElements = ElementCount;
@@ -82,11 +112,8 @@ public:
 		int32 EndIndex = Index + 1;
 		if (EndIndex > NumElements)
 		{
-			int32 NumChunks = (EndIndex + ChunkSize - 1) >> ChunkBits;
-			Chunks.SetNum(NumChunks);
+			SetNum(EndIndex, Default);
 		}
-
-		NumElements = EndIndex;
 	}
 
 	/** Fills the index with the default value */
@@ -111,7 +138,7 @@ public:
 		return Ar;
 	}
 
-	/** Gets the attribute array at the given index as a TArrayView */
+	/** Gets the array attribute at the given index as a TArrayView */
 	TArrayView<AttributeType> Get(int32 Index)
 	{
 		checkSlow(Index >= 0 && Index < NumElements);
@@ -121,7 +148,7 @@ public:
 		return TArrayView<AttributeType>(Chunk.Data.GetData() + Chunk.StartIndex[ElementIndex], Chunk.Count[ElementIndex]);
 	}
 
-	/** Gets the attribute array at the given index as a TArrayView */
+	/** Gets the array attribute at the given index as a TArrayView */
 	TArrayView<const AttributeType> Get(int32 Index) const
 	{
 		checkSlow(Index >= 0 && Index < NumElements);
@@ -131,10 +158,32 @@ public:
 		return TArrayView<const AttributeType>(Chunk.Data.GetData() + Chunk.StartIndex[ElementIndex], Chunk.Count[ElementIndex]);
 	}
 
+	/** Gets the element count of the array attribute at the given index */
+	int32 GetElementCount(int32 Index) const
+	{
+		checkSlow(Index >= 0 && Index < NumElements);
+		int32 ChunkIndex = Index >> ChunkBits;
+		int32 ElementIndex = Index & ChunkMask;
+		const FChunk& Chunk = Chunks[ChunkIndex];
+		return Chunk.Count[ElementIndex];
+	}
+
+	/** Gets the reserved element count of the array attribute at the given index */
+	int32 GetReservedElementCount(int32 Index) const
+	{
+		checkSlow(Index >= 0 && Index < NumElements);
+		int32 ChunkIndex = Index >> ChunkBits;
+		int32 ElementIndex = Index & ChunkMask;
+		const FChunk& Chunk = Chunks[ChunkIndex];
+		return Chunk.MaxCount[ElementIndex];
+	}
+
 	/** Sets the attribute array at the given index to the given TArrayView */
 	void Set(int32 Index, TArrayView<const AttributeType> Value)
 	{
-		TArrayView<AttributeType> Element = SetElementCount(Index, Value.Num(), false);
+		const bool bSetCount = true;
+		const bool bSetDefault = false;
+		TArrayView<AttributeType> Element = SetElementSize(Index, Value.Num(), bSetCount, bSetDefault);
 
 		for (int32 I = 0; I < Value.Num(); I++)
 		{
@@ -143,7 +192,7 @@ public:
 	}
 
 	/** Sets the given attribute array element to have the given number of subarray elements. */
-	TArrayView<AttributeType> SetElementCount(int32 Index, int32 Size, bool bSetDefault = true)
+	TArrayView<AttributeType> SetElementSize(int32 Index, int32 Size, bool bSetCount, bool bSetDefault = true)
 	{
 		checkSlow(Index >= 0 && Index < NumElements);
 		int32 ChunkIndex = Index >> ChunkBits;
@@ -157,14 +206,15 @@ public:
 		{
 			// If we are requesting that MaxCount grows for this element, we have to insert new data elements,
 			// and adjust subsequent start indices accordingly.
-			Chunk.Data.SetNum(Chunk.Data.Num() + Extra);
+			int32 OldDataSize = Chunk.Data.Num();
+			Chunk.Data.SetNum(OldDataSize + Extra);
 
-			AttributeType* StartPtr = BasePtr + Chunk.StartIndex[ElementIndex] + Extra;
-			AttributeType* EndPtr = BasePtr + Chunk.Data.Num();
-			FMemory::Memmove(StartPtr, StartPtr - Extra, (EndPtr - StartPtr) * sizeof(AttributeType));
+			AttributeType* StartPtr = BasePtr + Chunk.StartIndex[ElementIndex] + Chunk.MaxCount[ElementIndex];
+			AttributeType* EndPtr = BasePtr + OldDataSize;
+			FMemory::Memmove(StartPtr + Extra, StartPtr, (EndPtr - StartPtr) * sizeof(AttributeType));
 
 			Chunk.MaxCount[ElementIndex] = Size;
-			for (int32 I = ElementIndex + 1; I < ChunkSize; I++)
+			for (int32 I = ElementIndex + 1; I < Chunk.NumElements; I++)
 			{
 				Chunk.StartIndex[I] += Extra;
 			}
@@ -172,14 +222,28 @@ public:
 
 		if (bSetDefault)
 		{
-			// Set any newly created members to the default value
-			for (int32 I = Chunk.Count[ElementIndex]; I < Chunk.MaxCount[ElementIndex]; I++)
+			if (Size > Chunk.Count[ElementIndex])
 			{
-				Chunk.Data[Chunk.StartIndex[ElementIndex] + I] = DefaultValue;
+				// Set any newly created members to the default value
+				for (int32 I = Chunk.Count[ElementIndex]; I < Chunk.MaxCount[ElementIndex]; I++)
+				{
+					Chunk.Data[Chunk.StartIndex[ElementIndex] + I] = DefaultValue;
+				}
+			}
+			else
+			{
+				// Set any removed values to the default value
+				for (int32 I = Size; I < Chunk.Count[ElementIndex]; I++)
+				{
+					Chunk.Data[Chunk.StartIndex[ElementIndex] + I] = DefaultValue;
+				}
 			}
 		}
 
-		Chunk.Count[ElementIndex] = Size;
+		if (bSetCount)
+		{
+			Chunk.Count[ElementIndex] = Size;
+		}
 
 		return TArrayView<AttributeType>(BasePtr + Chunk.StartIndex[ElementIndex], Chunk.Count[ElementIndex]);
 	}
@@ -195,31 +259,56 @@ public:
 		checkSlow(SubArrayIndex >= 0 && SubArrayIndex <= CurrentCount);
 		int32 NewCount = CurrentCount + InsertCount;
 
-		// Resize the allocated space for this element if necessary
-		TArrayView<AttributeType> Element = SetElementCount(ElementIndex, NewCount, true);
+		if (NewCount > Chunk.MaxCount[ElementIndex])
+		{
+			// Add a few extra slots in anticipation of further growth in the future
+			// TODO: opt in to this behavior?
+			NewCount = NewCount * 4 / 3;
+
+			// Resize the allocated space for this element if necessary
+			const bool bSetCount = false;
+			const bool bSetDefault = true;
+			SetElementSize(Index, NewCount, bSetCount, bSetDefault);
+		}
+
+		Chunk.Count[ElementIndex] += InsertCount;
 
 		// Insert the new subarray elements
-		AttributeType* SpanStart = Element.GetData() + SubArrayIndex;
+		TArrayView<AttributeType> Element = Get(Index);
 		int32 SpanSize = NewCount - SubArrayIndex;
-		Algo::Rotate(TArrayView<AttributeType>(SpanStart, SpanSize), SpanSize - InsertCount);
+		TArrayView<AttributeType> ElementToRotate(Element.GetData() + SubArrayIndex, SpanSize);
+		Algo::Rotate(ElementToRotate, SpanSize - InsertCount);
 
 		return Element;
 	}
 
-	TArrayView<AttributeType> RemoveFromElement(int32 ElementIndex, int32 SubArrayIndex, int32 Count = 1)
+	TArrayView<AttributeType> RemoveFromElement(int32 Index, int32 SubArrayIndex, int32 Count = 1)
 	{
-		TArrayView<AttributeType> Element = Get(ElementIndex);
+		TArrayView<AttributeType> Element = Get(Index);
 		for (int32 I = SubArrayIndex + Count; I < Element.Num(); I++)
 		{
 			Element[I - Count] = Element[I];
 		}
 
-		return SetElementCount(ElementIndex, Element.Num() - Count, false);
+		const bool bSetCount = true;
+		const bool bSetDefault = true;
+		return SetElementSize(Index, Element.Num() - Count, bSetCount, bSetDefault);
 	}
 
+	AttributeType GetDefault() const { return DefaultValue; }
 
 private:
 
+	static int32 GetNumRequiredChunks(const int32 ElementCount)
+	{
+		return (ElementCount + ChunkSize - 1) >> ChunkBits;
+	}
+
+	static int32 GetElementsInLastChunk(const int32 ElementCount)
+	{
+		return ((ElementCount + ChunkSize - 1) & ChunkMask) + 1;
+	}
+
 	static const int32 ChunkBits = 8;
 	static const int32 ChunkSize = (1 << ChunkBits);
 	static const int32 ChunkMask = ChunkSize - 1;
@@ -234,9 +323,23 @@ private:
 		friend FArchive& operator <<(FArchive& Ar, FChunk& Chunk)
 		{
 			Ar << Chunk.Data;
-			Ar << Chunk.StartIndex;
-			Ar << Chunk.Count;
-			Ar << Chunk.MaxCount;
+			Ar << Chunk.NumElements;
+
+			for (int32 Index = 0; Index < Chunk.NumElements; Index++)
+			{
+				Ar << Chunk.StartIndex[Index];
+			}
+
+			for (int32 Index = 0; Index < Chunk.NumElements; Index++)
+			{
+				Ar << Chunk.Count[Index];
+			}
+
+			for (int32 Index = 0; Index < Chunk.NumElements; Index++)
+			{
+				Ar << Chunk.MaxCount[Index];
+			}
+
 			return Ar;
 		}
 
@@ -249,11 +352,16 @@ private:
 		TStaticArray<int32, ChunkSize> StartIndex = TStaticArray<int32, ChunkSize>(0);
 		TStaticArray<int32, ChunkSize> Count = TStaticArray<int32, ChunkSize>(0);
 		TStaticArray<int32, ChunkSize> MaxCount = TStaticArray<int32, ChunkSize>(0);
+
+		// Specify number of valid elements in this chunk.
+		// Default to maximum number so that we can add chunks in bulk, and all put the final
+		// chunk will already be correctly initialized.
+		int32 NumElements = ChunkSize;
 	};
 
 	TArray<FChunk> Chunks;
 	int32 NumElements = 0;
-	AttributeType DefaultValue;
+	AttributeType DefaultValue = AttributeType();
 };
 
 
@@ -368,25 +476,193 @@ public:
 	 */
 	void SetNum(int32 Num)
 	{
-		Array.SetElementCount(Index, Num, true);
+		const bool bSetCount = true;
+		const bool bSetDefault = true;
+		Array.SetElementSize(Index, Num, bSetCount, bSetDefault);
+	}
+
+
+	/**
+	 * Reserves the number of elements in the array attribute.
+	 */
+	void Reserve(int32 Num)
+	{
+		const bool bSetCount = false;
+		const bool bSetDefault = true;
+		Array.SetElementSize(Index, Num, bSetCount, bSetDefault);
+	}
+
+
+	/**
+	 * Adds the given value to the end of the array attribute, reserving extra space if necessary.
+	 */
+	int32 Add(const AttributeType& Value)
+	{
+		int32 ElementCount = Array.GetElementCount(Index);
+		TArrayView<AttributeType> Element = Array.InsertIntoElement(Index, ElementCount, 1);
+		Element[ElementCount] = Value;
+		return ElementCount;
+	}
+
+
+	/**
+	 * Adds the given value to the end of the array attribute, reserving extra space if necessary.
+	 */
+	int32 AddUnique(const AttributeType& Value)
+	{
+		int32 FoundIndex = Find(Value);
+		if (FoundIndex == INDEX_NONE)
+		{
+			return Add(Value);
+		}
+		else
+		{
+			return FoundIndex;
+		}
 	}
 
+
 	/**
-	 * Inserts a number of elements in the array attribute
+	 * Inserts a number of default-constructed elements in the array attribute
 	 */
-	void Insert(int32 StartIndex, int32 Count)
+	void InsertDefaulted(int32 StartIndex, int32 Count)
 	{
 		Array.InsertIntoElement(Index, StartIndex, Count);
 	}
 
+
+	/**
+	 * Inserts an element into the given index in the array attribute
+	 */
+	void Insert(const AttributeType& Value, int32 StartIndex)
+	{
+		TArrayView<AttributeType> Element = Array.InsertIntoElement(Index, StartIndex, 1);
+		Element[StartIndex] = Value;
+	}
+
+
 	/**
 	 * Removes a number of elements from the array attribute
 	 */
-	void Remove(int32 StartIndex, int32 Count)
+	void RemoveAt(int32 StartIndex, int32 Count)
 	{
 		Array.RemoveFromElement(Index, StartIndex, Count);
 	}
 
+
+	/**
+	 * Removes all elements which match the predicate
+	 */
+	template <typename Predicate>
+	int32 RemoveAll(Predicate Pred)
+	{
+		TArrayView<AttributeType> Element = Array.Get(Index);
+		int32 Count = Element.Num();
+
+		int32 WriteIndex = 0;
+		for (int32 ReadIndex = 0; ReadIndex < Count; ReadIndex++)
+		{
+			if (!Pred(Element[ReadIndex]))
+			{
+				Element[WriteIndex] = Element[ReadIndex];
+				WriteIndex++;
+			}
+		}
+
+		const bool bSetCount = true;
+		const bool bSetDefault = true;
+		Array.SetElementSize(Index, WriteIndex, bSetCount, bSetDefault);
+
+		return Count - WriteIndex;
+	}
+
+
+	int32 Remove(const AttributeType& Value)
+	{
+		return RemoveAll([&Value](const AttributeType& ToCompare) { return Value == ToCompare; });
+	}
+
+
+	/**
+	 * Finds the given value in the array attribute and returns its index
+	 */
+	int32 Find(const AttributeType& Value) const
+	{
+		return Array.Get(Index).Find(Value);
+	}
+
+
+	/**
+	 * Finds an element which matches a predicate functor.
+	 *
+	 * @param Pred The functor to apply to each element.
+	 *
+	 * @return Pointer to the first element for which the predicate returns
+	 *         true, or nullptr if none is found.
+	 */
+	template <typename Predicate>
+	AttributeType* FindByPredicate(Predicate Pred) const
+	{
+		return Array.Get(Index).FindByPredicate(MoveTemp(Pred));
+	}
+
+
+	/**
+	 * Finds an item by predicate.
+	 *
+	 * @param Pred The predicate to match.
+	 *
+	 * @returns Index to the first matching element, or INDEX_NONE if none is
+	 *          found.
+	 */
+	template <typename Predicate>
+	int32 IndexOfByPredicate(Predicate Pred) const
+	{
+		return Array.Get(Index).IndexOfByPredicate(MoveTemp(Pred));
+	}
+
+
+	/**
+	 * Determines whether the array attribute contains a particular value
+	 */
+	bool Contains(const AttributeType& Value) const
+	{
+		return Array.Get(Index).Contains(Value);
+	}
+
+
+	/**
+	 * Checks if this array contains element for which the predicate is true.
+	 *
+	 * @param Predicate to use
+	 *
+	 * @returns	True if found. False otherwise.
+	 */
+	template <typename Predicate>
+	bool ContainsByPredicate(Predicate Pred) const
+	{
+		return Array.Get(Index).ContainsByPredicate(MoveTemp(Pred));
+	}
+
+
+	/**
+	 * Sorts the array assuming < operator is defined for the item type.
+	 */
+	void Sort()
+	{
+		Array.Get(Index).Sort();
+	}
+
+
+	/**
+	 * Stable sorts the array assuming < operator is defined for the item type.
+	 */
+	void StableSort()
+	{
+		Array.Get(Index).StableSort();
+	}
+
+
 	/**
 	 * Return a TArrayView representing this array attribute.
 	 */
diff --git a/Engine/Source/Runtime/MeshDescription/Public/MeshAttributeArray.h b/Engine/Source/Runtime/MeshDescription/Public/MeshAttributeArray.h
--- a/Engine/Source/Runtime/MeshDescription/Public/MeshAttributeArray.h
+++ b/Engine/Source/Runtime/MeshDescription/Public/MeshAttributeArray.h
@@ -335,7 +335,7 @@ protected:
  * This is a type-specific attribute array, which is actually instanced in the attribute set.
  */
 template <typename AttributeType>
-class TMeshAttributeArraySet : public FMeshAttributeArraySetBase
+class TMeshAttributeArraySet final : public FMeshAttributeArraySetBase
 {
 	static_assert(!TIsArray<AttributeType>::Value, "TMeshAttributeArraySet must take a simple type.");
 
@@ -513,7 +513,7 @@ protected:
 * This is a type-specific attribute array, which is actually instanced in the attribute set.
 */
 template <typename AttributeType>
-class TMeshUnboundedAttributeArraySet : public FMeshAttributeArraySetBase
+class TMeshUnboundedAttributeArraySet final : public FMeshAttributeArraySetBase
 {
 	using Super = FMeshAttributeArraySetBase;
 
@@ -704,20 +704,16 @@ struct TMeshAttributesRefType : TMeshAttributesRefTypeBase<T>
 	using RefType = T;
 	using ConstRefType = const T;
 	using NonConstRefType = typename TRemoveCV<T>::Type;
-	using SetType = TMeshAttributeArraySet<typename TRemoveCV<T>::Type>;
-	using ConstSetType = const TMeshAttributeArraySet<typename TRemoveCV<T>::Type>;
 };
 
 template <typename T>
 struct TMeshAttributesRefType<TArrayView<T>> : TMeshAttributesRefTypeBase<T>
 {
-	static const uint32 MinExpectedExtent = 1;
+	static const uint32 MinExpectedExtent = 0;
 	static const uint32 MaxExpectedExtent = 0xFFFFFFFF;
 	using RefType = TArrayView<T>;
 	using ConstRefType = TArrayView<const T>;
 	using NonConstRefType = TArrayView<typename TRemoveCV<T>::Type>;
-	using SetType = TMeshAttributeArraySet<typename TRemoveCV<T>::Type>;
-	using ConstSetType = const TMeshAttributeArraySet<typename TRemoveCV<T>::Type>;
 };
 
 template <typename T>
@@ -728,8 +724,6 @@ struct TMeshAttributesRefType<TArrayAttribute<T>> : TMeshAttributesRefTypeBase<T
 	using RefType = TArrayAttribute<T>;
 	using ConstRefType = TArrayAttribute<const T>;
 	using NonConstRefType = TArrayAttribute<typename TRemoveCV<T>::Type>;
-	using SetType = TMeshUnboundedAttributeArraySet<typename TRemoveCV<T>::Type>;
-	using ConstSetType = const TMeshUnboundedAttributeArraySet<typename TRemoveCV<T>::Type>;
 };
 
 
@@ -779,191 +773,195 @@ template <typename AttributeType>
 using TMeshAttributesConstArray = TMeshAttributesRef<int32, typename TMeshAttributesRefType<AttributeType>::ConstRefType>;
 
 
+// This is the default implementation which handles simple attributes, i.e. those of a simple type T.
+// There are partial specializations which handle compound attributes below, i.e. those accessed via TArrayView<T> or TAttributeArray<T>
 template <typename ElementIDType, typename AttributeType>
 class TMeshAttributesRef
 {
 	template <typename T, typename U> friend class TMeshAttributesRef;
 
 public:
+	using BaseArrayType = typename TCopyQualifiersFromTo<AttributeType, FMeshAttributeArraySetBase>::Type;
 	using ArrayType = typename TCopyQualifiersFromTo<AttributeType, TMeshAttributeArraySet<typename TRemoveCV<AttributeType>::Type>>::Type;
 
-	/** Constructor taking a pointer to a TMeshAttributeArraySet */
-	FORCEINLINE explicit TMeshAttributesRef(ArrayType* InArrayPtr = nullptr, uint32 InExtent = 1)
+	/** Constructor taking a pointer to a FMeshAttributeArraySetBase */
+	explicit TMeshAttributesRef(BaseArrayType* InArrayPtr = nullptr, uint32 InExtent = 1)
 		: ArrayPtr(InArrayPtr)
 	{}
 
 	/** Implicitly construct a TMeshAttributesRef-to-const from a regular one */
 	template <typename T = AttributeType, typename TEnableIf<TIsSame<T, const T>::Value, int>::Type = 0>
-	FORCEINLINE TMeshAttributesRef(TMeshAttributesRef<ElementIDType, typename TRemoveCV<T>::Type> InRef)
+	TMeshAttributesRef(TMeshAttributesRef<ElementIDType, typename TRemoveCV<T>::Type> InRef)
 		: ArrayPtr(InRef.ArrayPtr)
 	{}
 
 	/** Implicitly construct a TMeshAttributesRef from a TMeshAttributesArray **/
 	template <typename T = ElementIDType, typename U = AttributeType, typename TEnableIf<!TIsSame<T, int32>::Value, int>::Type = 0>
-	FORCEINLINE TMeshAttributesRef(TMeshAttributesRef<int32, U> InRef)
+	TMeshAttributesRef(TMeshAttributesRef<int32, U> InRef)
 		: ArrayPtr(InRef.ArrayPtr)
 	{}
 
 	/** Implicitly construct a TMeshAttributesRef-to-const from a TMeshAttributesArray */
 	template <typename T = ElementIDType, typename U = AttributeType, typename TEnableIf<TIsSame<U, const U>::Value, int>::Type = 0, typename TEnableIf<!TIsSame<T, int32>::Value, int>::Type = 0>
-	FORCEINLINE TMeshAttributesRef(TMeshAttributesRef<int32, typename TRemoveCV<U>::Type> InRef)
+	TMeshAttributesRef(TMeshAttributesRef<int32, typename TRemoveCV<U>::Type> InRef)
 		: ArrayPtr(InRef.ArrayPtr)
 	{}
 
 	/** Access elements from attribute channel 0 */
 	template <typename T = ElementIDType, typename TEnableIf<TIsDerivedFrom<T, FElementID>::Value, int>::Type = 0>
-	FORCEINLINE AttributeType& operator[](const ElementIDType ElementID) const
+	AttributeType& operator[](const ElementIDType ElementID) const
 	{
-		return ArrayPtr->GetArrayForChannel(0).GetElementBase(ElementID.GetValue())[0];
+		return static_cast<ArrayType*>(ArrayPtr)->GetArrayForChannel(0).GetElementBase(ElementID.GetValue())[0];
 	}
 
 	/** Get the element with the given ID and channel */
 	template <typename T = ElementIDType, typename TEnableIf<TIsDerivedFrom<T, FElementID>::Value, int>::Type = 0>
-	FORCEINLINE AttributeType Get(const ElementIDType ElementID, const int32 Channel = 0) const
+	AttributeType Get(const ElementIDType ElementID, const int32 Channel = 0) const
 	{
-		return ArrayPtr->GetArrayForChannel(Channel).GetElementBase(ElementID.GetValue())[0];
+		return static_cast<ArrayType*>(ArrayPtr)->GetArrayForChannel(Channel).GetElementBase(ElementID.GetValue())[0];
 	}
 
-	FORCEINLINE AttributeType& operator[](int32 ElementIndex) const
+	AttributeType& operator[](int32 ElementIndex) const
 	{
-		return ArrayPtr->GetArrayForChannel(0).GetElementBase(ElementIndex)[0];
+		return static_cast<ArrayType*>(ArrayPtr)->GetArrayForChannel(0).GetElementBase(ElementIndex)[0];
 	}
 
-	FORCEINLINE AttributeType Get(int32 ElementIndex, const int32 Channel = 0) const
+	AttributeType Get(int32 ElementIndex, const int32 Channel = 0) const
 	{
-		return ArrayPtr->GetArrayForChannel(Channel).GetElementBase(ElementIndex)[0];
+		return static_cast<ArrayType*>(ArrayPtr)->GetArrayForChannel(Channel).GetElementBase(ElementIndex)[0];
 	}
 
-	FORCEINLINE TArrayView<AttributeType> GetArrayView(int32 ElementIndex, const int32 Channel = 0) const
+	TArrayView<AttributeType> GetArrayView(int32 ElementIndex, const int32 Channel = 0) const
 	{
-		return TArrayView<AttributeType>(ArrayPtr->GetArrayForChannel(Channel).GetElementBase(ElementIndex), 1);
+		return TArrayView<AttributeType>(static_cast<ArrayType*>(ArrayPtr)->GetArrayForChannel(Channel).GetElementBase(ElementIndex), 1);
 	}
 
-	FORCEINLINE TArrayView<AttributeType> GetRawArray(const int32 AttributeChannel = 0) const
+	TArrayView<AttributeType> GetRawArray(const int32 AttributeChannel = 0) const
 	{
 		if (ArrayPtr == nullptr)
 		{
 			return TArrayView<AttributeType>();
 		}
 
-		AttributeType* Element = ArrayPtr->GetArrayForChannel(AttributeChannel).GetElementBase(0);
+		AttributeType* Element = static_cast<ArrayType*>(ArrayPtr)->GetArrayForChannel(AttributeChannel).GetElementBase(0);
 		return TArrayView<AttributeType>(Element, GetNumElements());
 	}
 
 	/** Return whether the reference is valid or not */
-	FORCEINLINE bool IsValid() const { return (ArrayPtr != nullptr); }
+	bool IsValid() const { return (ArrayPtr != nullptr); }
 
 	/** Return default value for this attribute type */
-	FORCEINLINE AttributeType GetDefaultValue() const { return ArrayPtr->GetDefaultValue(); }
+	AttributeType GetDefaultValue() const { return static_cast<ArrayType*>(ArrayPtr)->GetDefaultValue(); }
 
 	UE_DEPRECATED(5.0, "Please use GetNumChannels().")
-	FORCEINLINE int32 GetNumIndices() const
+	int32 GetNumIndices() const
 	{
-		return ArrayPtr->ArrayType::GetNumChannels();	// note: override virtual dispatch
+		return static_cast<ArrayType*>(ArrayPtr)->ArrayType::GetNumChannels();	// note: override virtual dispatch
 	}
 
 	/** Return number of indices this attribute has */
-	FORCEINLINE int32 GetNumChannels() const
+	int32 GetNumChannels() const
 	{
-		return ArrayPtr->ArrayType::GetNumChannels();	// note: override virtual dispatch
+		return static_cast<ArrayType*>(ArrayPtr)->ArrayType::GetNumChannels();	// note: override virtual dispatch
 	}
 
 	/** Get the number of elements in this attribute array */
-	FORCEINLINE int32 GetNumElements() const
+	int32 GetNumElements() const
 	{
 		return ArrayPtr->GetNumElements();
 	}
 
 	/** Get the flags for this attribute array set */
-	FORCEINLINE EMeshAttributeFlags GetFlags() const { return ArrayPtr->GetFlags(); }
+	EMeshAttributeFlags GetFlags() const { return ArrayPtr->GetFlags(); }
 
-	FORCEINLINE uint32 GetExtent() const { return 1; }
+	/** Get the extent for this attribute type */
+	uint32 GetExtent() const { return 1; }
 
 	/** Set the element with the given ID and index 0 to the provided value */
 	template <typename T = ElementIDType, typename TEnableIf<TIsDerivedFrom<T, FElementID>::Value, int>::Type = 0>
-	FORCEINLINE void Set(const ElementIDType ElementID, const AttributeType& Value) const
+	void Set(const ElementIDType ElementID, const AttributeType& Value) const
 	{
-		ArrayPtr->GetArrayForChannel(0).GetElementBase(ElementID.GetValue())[0] = Value;
+		static_cast<ArrayType*>(ArrayPtr)->GetArrayForChannel(0).GetElementBase(ElementID.GetValue())[0] = Value;
 	}
 
 	/** Set the element with the given ID and channel to the provided value */
 	template <typename T = ElementIDType, typename TEnableIf<TIsDerivedFrom<T, FElementID>::Value, int>::Type = 0>
-	FORCEINLINE void Set(const ElementIDType ElementID, const int32 Channel, const AttributeType& Value) const
+	void Set(const ElementIDType ElementID, const int32 Channel, const AttributeType& Value) const
 	{
-		ArrayPtr->GetArrayForChannel(Channel).GetElementBase(ElementID.GetValue())[0] = Value;
+		static_cast<ArrayType*>(ArrayPtr)->GetArrayForChannel(Channel).GetElementBase(ElementID.GetValue())[0] = Value;
 	}
 
-	FORCEINLINE void Set(int32 ElementIndex, const AttributeType& Value) const
+	void Set(int32 ElementIndex, const AttributeType& Value) const
 	{
-		ArrayPtr->GetArrayForChannel(0).GetElementBase(ElementIndex)[0] = Value;
+		static_cast<ArrayType*>(ArrayPtr)->GetArrayForChannel(0).GetElementBase(ElementIndex)[0] = Value;
 	}
 
-	FORCEINLINE void Set(int32 ElementIndex, const int32 Channel, const AttributeType& Value) const
+	void Set(int32 ElementIndex, const int32 Channel, const AttributeType& Value) const
 	{
-		ArrayPtr->GetArrayForChannel(Channel).GetElementBase(ElementIndex)[0] = Value;
+		static_cast<ArrayType*>(ArrayPtr)->GetArrayForChannel(Channel).GetElementBase(ElementIndex)[0] = Value;
 	}
 
-	FORCEINLINE void SetArrayView(int32 ElementIndex, TArrayView<const AttributeType> Value) const
+	void SetArrayView(int32 ElementIndex, TArrayView<const AttributeType> Value) const
 	{
 		check(Value.Num() == 1);
-		ArrayPtr->GetArrayForChannel(0).GetElementBase(ElementIndex)[0] = Value[0];
+		static_cast<ArrayType*>(ArrayPtr)->GetArrayForChannel(0).GetElementBase(ElementIndex)[0] = Value[0];
 	}
 
-	FORCEINLINE void SetArrayView(int32 ElementIndex, const int32 Channel, TArrayView<const AttributeType> Value) const
+	void SetArrayView(int32 ElementIndex, const int32 Channel, TArrayView<const AttributeType> Value) const
 	{
 		check(Value.Num() == 1);
-		ArrayPtr->GetArrayForChannel(Channel).GetElementBase(ElementIndex)[0] = Value[0];
+		static_cast<ArrayType*>(ArrayPtr)->GetArrayForChannel(Channel).GetElementBase(ElementIndex)[0] = Value[0];
 	}
 
 	/** Copies the given attribute array and channel to this channel */
 	void Copy(TMeshAttributesRef<ElementIDType, const AttributeType> Src, const int32 DestChannel = 0, const int32 SrcChannel = 0);
 
 	UE_DEPRECATED(5.0, "Please use SetNumChannels().")
-	FORCEINLINE void SetNumIndices(const int32 NumChannels) const
+	void SetNumIndices(const int32 NumChannels) const
 	{
-		ArrayPtr->ArrayType::SetNumChannels(NumChannels);	// note: override virtual dispatch
+		static_cast<ArrayType*>(ArrayPtr)->ArrayType::SetNumChannels(NumChannels);	// note: override virtual dispatch
 	}
 
 	/** Sets number of channels this attribute has */
-	FORCEINLINE void SetNumChannels(const int32 NumChannels) const
+	void SetNumChannels(const int32 NumChannels) const
 	{
-		ArrayPtr->ArrayType::SetNumChannels(NumChannels);	// note: override virtual dispatch
+		static_cast<ArrayType*>(ArrayPtr)->ArrayType::SetNumChannels(NumChannels);	// note: override virtual dispatch
 	}
 
 	UE_DEPRECATED(5.0, "Please use InsertChannel().")
-	FORCEINLINE void InsertIndex(const int32 Index) const
+	void InsertIndex(const int32 Index) const
 	{
-		ArrayPtr->ArrayType::InsertChannel(Index);		// note: override virtual dispatch
+		static_cast<ArrayType*>(ArrayPtr)->ArrayType::InsertChannel(Index);		// note: override virtual dispatch
 	}
 
 	/** Inserts an attribute channel */
-	FORCEINLINE void InsertChannel(const int32 Index) const
+	void InsertChannel(const int32 Channel) const
 	{
-		ArrayPtr->ArrayType::InsertChannel(Index);		// note: override virtual dispatch
+		static_cast<ArrayType*>(ArrayPtr)->ArrayType::InsertChannel(Channel);		// note: override virtual dispatch
 	}
 
 	UE_DEPRECATED(5.0, "Please use RemoveChannel().")
-	FORCEINLINE void RemoveIndex(const int32 Index) const
+	void RemoveIndex(const int32 Index) const
 	{
-		ArrayPtr->ArrayType::RemoveChannel(Index);		// note: override virtual dispatch
+		static_cast<ArrayType*>(ArrayPtr)->ArrayType::RemoveChannel(Index);		// note: override virtual dispatch
 	}
 
 	/** Removes an attribute channel */
-	FORCEINLINE void RemoveChannel(const int32 Index) const
+	void RemoveChannel(const int32 Channel) const
 	{
-		ArrayPtr->ArrayType::RemoveChannel(Index);		// note: override virtual dispatch
+		static_cast<ArrayType*>(ArrayPtr)->ArrayType::RemoveChannel(Channel);		// note: override virtual dispatch
 	}
 
-protected:
-	ArrayType* ArrayPtr;
+private:
+	BaseArrayType* ArrayPtr;
 };
 
 template <typename ElementIDType, typename AttributeType>
 void TMeshAttributesRef<ElementIDType, AttributeType>::Copy(TMeshAttributesRef<ElementIDType, const AttributeType> Src, const int32 DestChannel, const int32 SrcChannel)
 {
 	check(Src.IsValid());
-	const TMeshAttributeArrayBase<AttributeType>& SrcArray = Src->ArrayPtr->GetArrayForChannel(SrcChannel);
-	TMeshAttributeArrayBase<AttributeType>& DestArray = ArrayPtr->GetArrayForChannel(DestChannel);
+	const TMeshAttributeArrayBase<AttributeType>& SrcArray = static_cast<const ArrayType*>(Src->ArrayPtr)->GetArrayForChannel(SrcChannel);
+	TMeshAttributeArrayBase<AttributeType>& DestArray = static_cast<ArrayType*>(ArrayPtr)->GetArrayForChannel(DestChannel);
 	const int32 Num = FMath::Min(SrcArray.Num(), DestArray.Num());
 	for (int32 Index = 0; Index < Num; Index++)
 	{
@@ -978,31 +976,33 @@ class TMeshAttributesRef<ElementIDType, TArrayView<AttributeType>>
 	template <typename T, typename U> friend class TMeshAttributesRef;
 
 public:
-	using ArrayType = typename TCopyQualifiersFromTo<AttributeType, TMeshAttributeArraySet<typename TRemoveCV<AttributeType>::Type>>::Type;
+	using BaseArrayType = typename TCopyQualifiersFromTo<AttributeType, FMeshAttributeArraySetBase>::Type;
+	using BoundedArrayType = typename TCopyQualifiersFromTo<AttributeType, TMeshAttributeArraySet<typename TRemoveCV<AttributeType>::Type>>::Type;
+	using UnboundedArrayType = typename TCopyQualifiersFromTo<AttributeType, TMeshUnboundedAttributeArraySet<typename TRemoveCV<AttributeType>::Type>>::Type;
 
 	/** Constructor taking a pointer to a TMeshAttributeArraySet */
-	FORCEINLINE explicit TMeshAttributesRef(ArrayType* InArrayPtr = nullptr, uint32 InExtent = 1)
+	explicit TMeshAttributesRef(BaseArrayType* InArrayPtr = nullptr, uint32 InExtent = 1)
 		: ArrayPtr(InArrayPtr),
 		  Extent(InExtent)
 	{}
 
 	/** Implicitly construct a TMeshAttributesRef-to-const from a regular one */
 	template <typename T = AttributeType, typename TEnableIf<TIsSame<T, const T>::Value, int>::Type = 0>
-	FORCEINLINE TMeshAttributesRef(TMeshAttributesRef<ElementIDType, TArrayView<typename TRemoveCV<T>::Type>> InRef)
+	TMeshAttributesRef(TMeshAttributesRef<ElementIDType, TArrayView<typename TRemoveCV<T>::Type>> InRef)
 		: ArrayPtr(InRef.ArrayPtr),
 		  Extent(InRef.Extent)
 	{}
 
 	/** Implicitly construct a TMeshAttributesRef from a TMeshAttributesArray **/
 	template <typename T = ElementIDType, typename U = AttributeType, typename TEnableIf<!TIsSame<T, int32>::Value, int>::Type = 0>
-	FORCEINLINE TMeshAttributesRef(TMeshAttributesRef<int32, TArrayView<U>> InRef)
+	TMeshAttributesRef(TMeshAttributesRef<int32, TArrayView<U>> InRef)
 		: ArrayPtr(InRef.ArrayPtr),
 		  Extent(InRef.Extent)
 	{}
 
 	/** Implicitly construct a TMeshAttributesRef-to-const from a TMeshAttributesArray */
 	template <typename T = ElementIDType, typename U = AttributeType, typename TEnableIf<TIsSame<U, const U>::Value, int>::Type = 0, typename TEnableIf<!TIsSame<T, int32>::Value, int>::Type = 0>
-	FORCEINLINE TMeshAttributesRef(TMeshAttributesRef<int32, TArrayView<typename TRemoveCV<U>::Type>> InRef)
+	TMeshAttributesRef(TMeshAttributesRef<int32, TArrayView<typename TRemoveCV<U>::Type>> InRef)
 		: ArrayPtr(InRef.ArrayPtr),
 		  Extent(InRef.Extent)
 	{}
@@ -1010,128 +1010,166 @@ public:
 
 	/** Access elements from attribute channel 0 */
 	template <typename T = ElementIDType, typename TEnableIf<TIsDerivedFrom<T, FElementID>::Value, int>::Type = 0>
-	FORCEINLINE TArrayView<AttributeType> operator[](const ElementIDType ElementID) const
+	TArrayView<AttributeType> operator[](const ElementIDType ElementID) const
 	{
-		AttributeType* Element = ArrayPtr->GetArrayForChannel(0).GetElementBase(ElementID.GetValue());
-		return TArrayView<AttributeType>(Element, Extent);
+		if (Extent > 0)
+		{
+			AttributeType* Element = static_cast<BoundedArrayType*>(ArrayPtr)->GetArrayForChannel(0).GetElementBase(ElementID.GetValue());
+			return TArrayView<AttributeType>(Element, Extent);
+		}
+		else
+		{
+			return static_cast<UnboundedArrayType*>(ArrayPtr)->GetArrayForChannel(0).Get(ElementID.GetValue());
+		}
 	}
 
 	/** Get the element with the given ID and channel */
 	template <typename T = ElementIDType, typename TEnableIf<TIsDerivedFrom<T, FElementID>::Value, int>::Type = 0>
-	FORCEINLINE TArrayView<AttributeType> Get(const ElementIDType ElementID, const int32 Channel = 0) const
+	TArrayView<AttributeType> Get(const ElementIDType ElementID, const int32 Channel = 0) const
 	{
-		AttributeType* Element = ArrayPtr->GetArrayForChannel(Channel).GetElementBase(ElementID.GetValue());
-		return TArrayView<AttributeType>(Element, Extent);
+		if (Extent > 0)
+		{
+			AttributeType* Element = static_cast<BoundedArrayType*>(ArrayPtr)->GetArrayForChannel(Channel).GetElementBase(ElementID.GetValue());
+			return TArrayView<AttributeType>(Element, Extent);
+		}
+		else
+		{
+			return static_cast<UnboundedArrayType*>(ArrayPtr)->GetArrayForChannel(Channel).Get(ElementID.GetValue());
+		}
 	}
 
-	FORCEINLINE TArrayView<AttributeType> operator[](int32 ElementIndex) const
+	TArrayView<AttributeType> operator[](int32 ElementIndex) const
 	{
-		AttributeType* Element = ArrayPtr->GetArrayForChannel(0).GetElementBase(ElementIndex);
-		return TArrayView<AttributeType>(Element, Extent);
+		if (Extent > 0)
+		{
+			AttributeType* Element = static_cast<BoundedArrayType*>(ArrayPtr)->GetArrayForChannel(0).GetElementBase(ElementIndex);
+			return TArrayView<AttributeType>(Element, Extent);
+		}
+		else
+		{
+			return static_cast<UnboundedArrayType*>(ArrayPtr)->GetArrayForChannel(0).Get(ElementIndex);
+		}
 	}
 
-	FORCEINLINE TArrayView<AttributeType> Get(int32 ElementIndex, const int32 Channel = 0) const
+	TArrayView<AttributeType> Get(int32 ElementIndex, const int32 Channel = 0) const
 	{
-		AttributeType* Element = ArrayPtr->GetArrayForChannel(Channel).GetElementBase(ElementIndex);
-		return TArrayView<AttributeType>(Element, Extent);
+		if (Extent > 0)
+		{
+			AttributeType* Element = static_cast<BoundedArrayType*>(ArrayPtr)->GetArrayForChannel(Channel).GetElementBase(ElementIndex);
+			return TArrayView<AttributeType>(Element, Extent);
+		}
+		else
+		{
+			return static_cast<UnboundedArrayType*>(ArrayPtr)->GetArrayForChannel(Channel).Get(ElementIndex);
+		}
 	}
 
-	FORCEINLINE TArrayView<AttributeType> GetArrayView(int32 ElementIndex, const int32 Channel = 0) const
+	TArrayView<AttributeType> GetArrayView(int32 ElementIndex, const int32 Channel = 0) const
 	{
-		return TArrayView<AttributeType>(ArrayPtr->GetArrayForChannel(Channel).GetElementBase(ElementIndex), Extent);
+		if (Extent > 0)
+		{
+			return TArrayView<AttributeType>(static_cast<BoundedArrayType*>(ArrayPtr)->GetArrayForChannel(Channel).GetElementBase(ElementIndex), Extent);
+		}
+		else
+		{
+			return static_cast<UnboundedArrayType*>(ArrayPtr)->GetArrayForChannel(Channel).Get(ElementIndex);
+		}
 	}
 
-	FORCEINLINE TArrayView<AttributeType> GetRawArray(const int32 ChannelIndex = 0) const
+	TArrayView<AttributeType> GetRawArray(const int32 ChannelIndex = 0) const
 	{
+		// Can't get the attribute set raw array for unbounded arrays because they are chunked
+		check(Extent > 0);
+
 		if (ArrayPtr == nullptr)
 		{
 			return TArrayView<AttributeType>();
 		}
 
-		AttributeType* Element = ArrayPtr->GetArrayForChannel(ChannelIndex).GetElementBase(0);
+		AttributeType* Element = static_cast<BoundedArrayType*>(ArrayPtr)->GetArrayForChannel(ChannelIndex).GetElementBase(0);
 		return TArrayView<AttributeType>(Element, GetNumElements() * Extent);
 	}
 
 	/** Return whether the reference is valid or not */
-	FORCEINLINE bool IsValid() const { return (ArrayPtr != nullptr); }
+	bool IsValid() const { return (ArrayPtr != nullptr); }
 
 	/** Return default value for this attribute type */
-	FORCEINLINE AttributeType GetDefaultValue() const { return ArrayPtr->GetDefaultValue(); }
+	AttributeType GetDefaultValue() const { return static_cast<BoundedArrayType*>(ArrayPtr)->GetDefaultValue(); }
 
 	UE_DEPRECATED(5.0, "Please use GetNumChannels().")
-	FORCEINLINE int32 GetNumIndices() const
+	int32 GetNumIndices() const
 	{
-		return ArrayPtr->ArrayType::GetNumChannels();	// note: override virtual dispatch
+		return ArrayPtr->GetNumChannels();
 	}
 
 	/** Return number of channels this attribute has */
-	FORCEINLINE int32 GetNumChannels() const
+	int32 GetNumChannels() const
 	{
-		return ArrayPtr->ArrayType::GetNumChannels();	// note: override virtual dispatch
+		return ArrayPtr->GetNumChannels();
 	}
 
 	/** Get the number of elements in this attribute array */
-	FORCEINLINE int32 GetNumElements() const
+	int32 GetNumElements() const
 	{
 		return ArrayPtr->GetNumElements();
 	}
 
 	/** Get the flags for this attribute array set */
-	FORCEINLINE EMeshAttributeFlags GetFlags() const { return ArrayPtr->GetFlags(); }
+	EMeshAttributeFlags GetFlags() const { return ArrayPtr->GetFlags(); }
 
 	/** Return the extent of this attribute, i.e. the number of array elements it comprises */
-	FORCEINLINE uint32 GetExtent() const { return Extent; }
+	uint32 GetExtent() const { return Extent; }
 
 	/** Set the element with the given ID and index 0 to the provided value */
 	template <typename T = ElementIDType, typename TEnableIf<TIsDerivedFrom<T, FElementID>::Value, int>::Type = 0>
-	FORCEINLINE void Set(const ElementIDType ElementID, TArrayView<const AttributeType> Value) const
+	void Set(const ElementIDType ElementID, TArrayView<const AttributeType> Value) const
 	{
-		check(Value.Num() == Extent);
-		AttributeType* Element = ArrayPtr->GetArrayForChannel(0).GetElementBase(ElementID.GetValue());
+		TArrayView<AttributeType> Elements = Get(ElementID);
+		check(Value.Num() == Elements.Num());
 		for (uint32 Index = 0; Index < Extent; Index++)
 		{
-			Element[Index] = Value[Index];
+			Elements[Index] = Value[Index];
 		}
 	}
 
 	/** Set the element with the given ID and channel to the provided value */
 	template <typename T = ElementIDType, typename TEnableIf<TIsDerivedFrom<T, FElementID>::Value, int>::Type = 0>
-	FORCEINLINE void Set(const ElementIDType ElementID, const int32 Channel, TArrayView<const AttributeType> Value) const
+	void Set(const ElementIDType ElementID, const int32 Channel, TArrayView<const AttributeType> Value) const
 	{
-		check(Value.Num() == Extent);
-		AttributeType* Element = ArrayPtr->GetArrayForChannel(Channel).GetElementBase(ElementID.GetValue());
+		TArrayView<AttributeType> Elements = Get(ElementID, Channel);
+		check(Value.Num() == Elements.Num());
 		for (uint32 Index = 0; Index < Extent; Index++)
 		{
-			Element[Index] = Value[Index];
+			Elements[Index] = Value[Index];
 		}
 	}
 
-	FORCEINLINE void Set(int32 ElementIndex, TArrayView<const AttributeType> Value) const
+	void Set(int32 ElementIndex, TArrayView<const AttributeType> Value) const
 	{
-		check(Value.Num() == Extent);
-		AttributeType* Element = ArrayPtr->GetArrayForChannel(0).GetElementBase(ElementIndex);
+		TArrayView<AttributeType> Elements = Get(ElementIndex);
+		check(Value.Num() == Elements.Num());
 		for (uint32 Index = 0; Index < Extent; Index++)
 		{
-			Element[Index] = Value[Index];
+			Elements[Index] = Value[Index];
 		}
 	}
 
-	FORCEINLINE void Set(int32 ElementIndex, const int32 Channel, TArrayView<const AttributeType> Value) const
+	void Set(int32 ElementIndex, const int32 Channel, TArrayView<const AttributeType> Value) const
 	{
-		check(Value.Num() == Extent);
-		AttributeType* Element = ArrayPtr->GetArrayForChannel(Channel).GetElementBase(ElementIndex);
+		TArrayView<AttributeType> Elements = Get(ElementIndex, Channel);
+		check(Value.Num() == Elements.Num());
 		for (uint32 Index = 0; Index < Extent; Index++)
 		{
-			Element[Index] = Value[Index];
+			Elements[Index] = Value[Index];
 		}
 	}
 
-	FORCEINLINE void SetArrayView(int32 ElementIndex, TArrayView<const AttributeType> Value) const
+	void SetArrayView(int32 ElementIndex, TArrayView<const AttributeType> Value) const
 	{
 		Set(ElementIndex, Value);
 	}
 
-	FORCEINLINE void SetArrayView(int32 ElementIndex, const int32 Channel, TArrayView<const AttributeType> Value) const
+	void SetArrayView(int32 ElementIndex, const int32 Channel, TArrayView<const AttributeType> Value) const
 	{
 		Set(ElementIndex, Channel, Value);
 	}
@@ -1140,53 +1178,54 @@ public:
 	void Copy(TMeshAttributesConstRef<ElementIDType, TArrayView<AttributeType>> Src, const int32 DestChannel = 0, const int32 SrcChannel = 0);
 
 	UE_DEPRECATED(5.0, "Please use SetNumChannels().")
-	FORCEINLINE void SetNumIndices(const int32 NumChannels) const
+	void SetNumIndices(const int32 NumChannels) const
 	{
-		ArrayPtr->ArrayType::SetNumChannels(NumChannels);	// note: override virtual dispatch
+		ArrayPtr->SetNumChannels(NumChannels);
 	}
 
 	/** Sets number of channels this attribute has */
-	FORCEINLINE void SetNumChannels(const int32 NumChannels) const
+	void SetNumChannels(const int32 NumChannels) const
 	{
-		ArrayPtr->ArrayType::SetNumChannels(NumChannels);	// note: override virtual dispatch
+		ArrayPtr->SetNumChannels(NumChannels);
 	}
 
 	UE_DEPRECATED(5.0, "Please use InsertChannel().")
-	FORCEINLINE void InsertIndex(const int32 Index) const
+	void InsertIndex(const int32 Index) const
 	{
-		ArrayPtr->ArrayType::InsertChannel(Index);		// note: override virtual dispatch
+		ArrayPtr->InsertChannel(Index);
 	}
 
 	/** Inserts an attribute channel */
-	FORCEINLINE void InsertChannel(const int32 Index) const
+	void InsertChannel(const int32 Channel) const
 	{
-		ArrayPtr->ArrayType::InsertChannel(Index);		// note: override virtual dispatch
+		ArrayPtr->InsertChannel(Channel);
 	}
 
 	UE_DEPRECATED(5.0, "Please use RemoveChannel().")
-	FORCEINLINE void RemoveIndex(const int32 Index) const
+	void RemoveIndex(const int32 Index) const
 	{
-		ArrayPtr->ArrayType::RemoveChannel(Index);		// note: override virtual dispatch
+		ArrayPtr->RemoveChannel(Index);
 	}
 
 	/** Removes an attribute channel */
-	FORCEINLINE void RemoveChannel(const int32 Index) const
+	void RemoveChannel(const int32 Channel) const
 	{
-		ArrayPtr->ArrayType::RemoveChannel(Index);		// note: override virtual dispatch
+		ArrayPtr->RemoveChannel(Channel);
 	}
 
-protected:
-	ArrayType* ArrayPtr;
+private:
+	BaseArrayType* ArrayPtr;
 	uint32 Extent;
 };
 
 template <typename ElementIDType, typename AttributeType>
 void TMeshAttributesRef<ElementIDType, TArrayView<AttributeType>>::Copy(TMeshAttributesConstRef<ElementIDType, TArrayView<AttributeType>> Src, const int32 DestChannel, const int32 SrcChannel)
 {
+	check(Extent > 0);
 	check(Src.IsValid());
 	check(Src.Extent == Extent);
-	const TMeshAttributeArrayBase<AttributeType>& SrcArray = Src->ArrayPtr->GetArrayForChannel(SrcChannel);
-	TMeshAttributeArrayBase<AttributeType>& DestArray = ArrayPtr->GetArrayForChannel(DestChannel);
+	const TMeshAttributeArrayBase<AttributeType>& SrcArray = static_cast<const BoundedArrayType*>(Src->ArrayPtr)->GetArrayForChannel(SrcChannel);
+	TMeshAttributeArrayBase<AttributeType>& DestArray = static_cast<BoundedArrayType*>(ArrayPtr)->GetArrayForChannel(DestChannel);
 	const int32 Num = FMath::Min(SrcArray.Num(), DestArray.Num());
 	for (int32 Index = 0; Index < Num; Index++)
 	{
@@ -1204,140 +1243,151 @@ class TMeshAttributesRef<ElementIDType, TArrayAttribute<AttributeType>>
 	template <typename T, typename U> friend class TMeshAttributesRef;
 
 public:
+	using BaseArrayType = typename TCopyQualifiersFromTo<AttributeType, FMeshAttributeArraySetBase>::Type;
 	using ArrayType = typename TCopyQualifiersFromTo<AttributeType, TMeshUnboundedAttributeArraySet<typename TRemoveCV<AttributeType>::Type>>::Type;
 
 	/** Constructor taking a pointer to a TMeshUnboundedAttributeArraySet */
-	FORCEINLINE explicit TMeshAttributesRef(ArrayType* InArrayPtr = nullptr, uint32 InExtent = 0)
+	explicit TMeshAttributesRef(BaseArrayType* InArrayPtr = nullptr, uint32 InExtent = 0)
 		: ArrayPtr(InArrayPtr)
 	{}
 
 	/** Implicitly construct a TMeshAttributesRef-to-const from a regular one */
 	template <typename T = AttributeType, typename TEnableIf<TIsSame<T, const T>::Value, int>::Type = 0>
-	FORCEINLINE TMeshAttributesRef(TMeshAttributesRef<ElementIDType, TArrayAttribute<typename TRemoveCV<T>::Type>> InRef)
+	TMeshAttributesRef(TMeshAttributesRef<ElementIDType, TArrayAttribute<typename TRemoveCV<T>::Type>> InRef)
 		: ArrayPtr(InRef.ArrayPtr)
 	{}
 
 	/** Implicitly construct a TMeshAttributesRef from a TMeshAttributesArray **/
 	template <typename T = ElementIDType, typename U = AttributeType, typename TEnableIf<!TIsSame<T, int32>::Value, int>::Type = 0>
-	FORCEINLINE TMeshAttributesRef(TMeshAttributesRef<int32, TArrayAttribute<U>> InRef)
+	TMeshAttributesRef(TMeshAttributesRef<int32, TArrayAttribute<U>> InRef)
 		: ArrayPtr(InRef.ArrayPtr)
 	{}
 
 	/** Implicitly construct a TMeshAttributesRef-to-const from a TMeshAttributesArray */
 	template <typename T = ElementIDType, typename U = AttributeType, typename TEnableIf<TIsSame<U, const U>::Value, int>::Type = 0, typename TEnableIf<!TIsSame<T, int32>::Value, int>::Type = 0>
-	FORCEINLINE TMeshAttributesRef(TMeshAttributesRef<int32, TArrayAttribute<typename TRemoveCV<U>::Type>> InRef)
+	TMeshAttributesRef(TMeshAttributesRef<int32, TArrayAttribute<typename TRemoveCV<U>::Type>> InRef)
 		: ArrayPtr(InRef.ArrayPtr)
 	{}
 
 
 	/** Access elements from attribute channel 0 */
 	template <typename T = ElementIDType, typename TEnableIf<TIsDerivedFrom<T, FElementID>::Value, int>::Type = 0>
-	FORCEINLINE TArrayAttribute<AttributeType> operator[](const ElementIDType ElementID) const
+	TArrayAttribute<AttributeType> operator[](const ElementIDType ElementID) const
 	{
-		return TArrayAttribute<AttributeType>(ArrayPtr->GetArrayForChannel(0), ElementID.GetValue());
+		return TArrayAttribute<AttributeType>(static_cast<ArrayType*>(ArrayPtr)->GetArrayForChannel(0), ElementID.GetValue());
 	}
 
 	/** Get the element with the given ID and channel */
 	template <typename T = ElementIDType, typename TEnableIf<TIsDerivedFrom<T, FElementID>::Value, int>::Type = 0>
-	FORCEINLINE TArrayAttribute<AttributeType> Get(const ElementIDType ElementID, const int32 Channel = 0) const
+	TArrayAttribute<AttributeType> Get(const ElementIDType ElementID, const int32 Channel = 0) const
 	{
-		return TArrayAttribute<AttributeType>(ArrayPtr->GetArrayForChannel(Channel), ElementID.GetValue());
+		return TArrayAttribute<AttributeType>(static_cast<ArrayType*>(ArrayPtr)->GetArrayForChannel(Channel), ElementID.GetValue());
 	}
 
-	FORCEINLINE TArrayAttribute<AttributeType> operator[](int32 ElementIndex) const
+	TArrayAttribute<AttributeType> operator[](int32 ElementIndex) const
 	{
-		return TArrayAttribute<AttributeType>(ArrayPtr->GetArrayForChannel(0), ElementIndex);
+		return TArrayAttribute<AttributeType>(static_cast<ArrayType*>(ArrayPtr)->GetArrayForChannel(0), ElementIndex);
 	}
 
-	FORCEINLINE TArrayAttribute<AttributeType> Get(int32 ElementIndex, const int32 Channel = 0) const
+	TArrayAttribute<AttributeType> Get(int32 ElementIndex, const int32 Channel = 0) const
 	{
-		return TArrayAttribute<AttributeType>(ArrayPtr->GetArrayForChannel(Channel), ElementIndex);
+		return TArrayAttribute<AttributeType>(static_cast<ArrayType*>(ArrayPtr)->GetArrayForChannel(Channel), ElementIndex);
 	}
 
-	FORCEINLINE TArrayView<AttributeType> GetArrayView(int32 ElementIndex, const int32 Channel = 0) const
+	TArrayView<AttributeType> GetArrayView(int32 ElementIndex, const int32 Channel = 0) const
 	{
 		return Get(ElementIndex, Channel).ToArrayView();
 	}
 
+	/** In this specialization, GetRawArray returns a pointer to the attribute array container holding the attributes and their index pointers */
+	const TAttributeArrayContainer<AttributeType>* GetRawArray(const int32 AttributeChannel = 0) const
+	{
+		if (ArrayPtr == nullptr)
+		{
+			return nullptr;
+		}
+
+		return &static_cast<ArrayType*>(ArrayPtr)->GetArrayForChannel(AttributeChannel);
+	}
+
 	/** Return whether the reference is valid or not */
-	FORCEINLINE bool IsValid() const { return (ArrayPtr != nullptr); }
+	bool IsValid() const { return (ArrayPtr != nullptr); }
 
 	/** Return default value for this attribute type */
-	FORCEINLINE AttributeType GetDefaultValue() const { return ArrayPtr->GetDefaultValue(); }
+	AttributeType GetDefaultValue() const { return static_cast<ArrayType*>(ArrayPtr)->GetDefaultValue(); }
 
 	/** Return number of channels this attribute has */
-	FORCEINLINE int32 GetNumChannels() const
+	int32 GetNumChannels() const
 	{
-		return ArrayPtr->ArrayType::GetNumChannels();	// note: override virtual dispatch
+		return static_cast<ArrayType*>(ArrayPtr)->ArrayType::GetNumChannels();	// note: override virtual dispatch
 	}
 
 	/** Get the number of elements in this attribute array */
-	FORCEINLINE int32 GetNumElements() const
+	int32 GetNumElements() const
 	{
 		return ArrayPtr->GetNumElements();
 	}
 
 	/** Get the flags for this attribute array set */
-	FORCEINLINE EMeshAttributeFlags GetFlags() const { return ArrayPtr->GetFlags(); }
+	EMeshAttributeFlags GetFlags() const { return ArrayPtr->GetFlags(); }
 
 	/** Set the element with the given ID and index 0 to the provided value */
 	template <typename T = ElementIDType, typename TEnableIf<TIsDerivedFrom<T, FElementID>::Value, int>::Type = 0>
-	FORCEINLINE void Set(const ElementIDType ElementID, TArrayAttribute<const AttributeType> Value) const
+	void Set(const ElementIDType ElementID, TArrayAttribute<const AttributeType> Value) const
 	{
-		ArrayPtr->GetArrayForChannel(0).Set(ElementID.GetValue(), Value.ToArrayView());
+		static_cast<ArrayType*>(ArrayPtr)->GetArrayForChannel(0).Set(ElementID.GetValue(), Value.ToArrayView());
 	}
 
 	/** Set the element with the given ID and channel to the provided value */
 	template <typename T = ElementIDType, typename TEnableIf<TIsDerivedFrom<T, FElementID>::Value, int>::Type = 0>
-	FORCEINLINE void Set(const ElementIDType ElementID, const int32 Channel, TArrayAttribute<const AttributeType> Value) const
+	void Set(const ElementIDType ElementID, const int32 Channel, TArrayAttribute<const AttributeType> Value) const
 	{
-		ArrayPtr->GetArrayForChannel(Channel).Set(ElementID.GetValue(), Value.ToArrayView());
+		static_cast<ArrayType*>(ArrayPtr)->GetArrayForChannel(Channel).Set(ElementID.GetValue(), Value.ToArrayView());
 	}
 
-	FORCEINLINE void Set(int32 ElementIndex, TArrayAttribute<const AttributeType> Value) const
+	void Set(int32 ElementIndex, TArrayAttribute<const AttributeType> Value) const
 	{
-		ArrayPtr->GetArrayForChannel(0).Set(ElementIndex, Value.ToArrayView());
+		static_cast<ArrayType*>(ArrayPtr)->GetArrayForChannel(0).Set(ElementIndex, Value.ToArrayView());
 	}
 
-	FORCEINLINE void Set(int32 ElementIndex, const int32 Channel, TArrayAttribute<const AttributeType> Value) const
+	void Set(int32 ElementIndex, const int32 Channel, TArrayAttribute<const AttributeType> Value) const
 	{
-		ArrayPtr->GetArrayForChannel(Channel).Set(ElementIndex, Value.ToArrayView());
+		static_cast<ArrayType*>(ArrayPtr)->GetArrayForChannel(Channel).Set(ElementIndex, Value.ToArrayView());
 	}
 
-	FORCEINLINE void SetArrayView(int32 ElementIndex, TArrayView<const AttributeType> Value) const
+	void SetArrayView(int32 ElementIndex, TArrayView<const AttributeType> Value) const
 	{
-		ArrayPtr->GetArrayForChannel(0).Set(ElementIndex, Value);
+		static_cast<ArrayType*>(ArrayPtr)->GetArrayForChannel(0).Set(ElementIndex, Value);
 	}
 
-	FORCEINLINE void SetArrayView(int32 ElementIndex, const int32 Channel, TArrayView<const AttributeType> Value) const
+	void SetArrayView(int32 ElementIndex, const int32 Channel, TArrayView<const AttributeType> Value) const
 	{
-		ArrayPtr->GetArrayForChannel(Channel).Set(ElementIndex, Value);
+		static_cast<ArrayType*>(ArrayPtr)->GetArrayForChannel(Channel).Set(ElementIndex, Value);
 	}
 
 	/** Copies the given attribute array and index to this index */
 	void Copy(TMeshAttributesConstRef<ElementIDType, TArrayAttribute<AttributeType>> Src, const int32 DestChannel = 0, const int32 SrcChannel = 0);
 
 	/** Sets number of channels this attribute has */
-	FORCEINLINE void SetNumChannels(const int32 NumChannels) const
+	void SetNumChannels(const int32 NumChannels) const
 	{
-		ArrayPtr->ArrayType::SetNumChannels(NumChannels);	// note: override virtual dispatch
+		static_cast<ArrayType*>(ArrayPtr)->ArrayType::SetNumChannels(NumChannels);	// note: override virtual dispatch
 	}
 
 	/** Inserts an attribute channel */
-	FORCEINLINE void InsertChannel(const int32 Index) const
+	void InsertChannel(const int32 Index) const
 	{
-		ArrayPtr->ArrayType::InsertChannel(Index);		// note: override virtual dispatch
+		static_cast<ArrayType*>(ArrayPtr)->ArrayType::InsertChannel(Index);		// note: override virtual dispatch
 	}
 
 	/** Removes an attribute channel */
-	FORCEINLINE void RemoveChannel(const int32 Index) const
+	void RemoveChannel(const int32 Index) const
 	{
-		ArrayPtr->ArrayType::RemoveChannel(Index);		// note: override virtual dispatch
+		static_cast<ArrayType*>(ArrayPtr)->ArrayType::RemoveChannel(Index);		// note: override virtual dispatch
 	}
 
-protected:
-	ArrayType* ArrayPtr;
-	uint32 Extent;
+private:
+	BaseArrayType* ArrayPtr;
 };
 
 
@@ -1450,16 +1500,15 @@ public:
 		using AttributeType = typename TMeshAttributesRegisterType<T>::AttributeType;
 		using RealAttributeType = typename TMeshAttributesRegisterType<T>::RealAttributeType;
 		using RefType = typename TMeshAttributesRegisterType<T>::RefType;
-		using SetType = typename TMeshAttributesRegisterType<T>::SetType;
 		const uint32 Extent = TMeshAttributesRegisterType<T>::Extent;
 
 		if (FAttributesSetEntry* ArraySetPtr = Map.Find(AttributeName))
 		{
 			if ((*ArraySetPtr)->HasType<RealAttributeType>() && (*ArraySetPtr)->GetExtent() == Extent)
 			{
-				static_cast<SetType*>(ArraySetPtr->Get())->SetType::SetNumChannels(NumberOfChannels);	// note: override virtual dispatch
+				(*ArraySetPtr)->SetNumChannels(NumberOfChannels);
 				(*ArraySetPtr)->SetFlags(Flags);
-				return TMeshAttributesArray<RefType>(static_cast<SetType*>(ArraySetPtr->Get()), Extent);
+				return TMeshAttributesArray<RefType>(ArraySetPtr->Get(), Extent);
 			}
 			else
 			{
@@ -1468,7 +1517,7 @@ public:
 		}
 
 		FAttributesSetEntry& Entry = Map.Emplace(AttributeName, FAttributesSetEntry(NumberOfChannels, Default, Flags, NumElements, Extent));
-		return TMeshAttributesArray<RefType>(static_cast<SetType*>(Entry.Get()), Extent);
+		return TMeshAttributesArray<RefType>(Entry.Get(), Extent);
 	}
 
 	/**
@@ -1666,7 +1715,7 @@ public:
 				uint32 ActualExtent = (*ArraySetPtr)->GetExtent();
 				if (ActualExtent >= TMeshAttributesRefType<T>::MinExpectedExtent && ActualExtent <= TMeshAttributesRefType<T>::MaxExpectedExtent)
 				{
-					return TMeshAttributesConstRef<int32, RefType>(static_cast<typename TMeshAttributesRefType<T>::ConstSetType*>(ArraySetPtr->Get()), ActualExtent);
+					return TMeshAttributesConstRef<int32, RefType>(ArraySetPtr->Get(), ActualExtent);
 				}
 			}
 		}
@@ -1689,7 +1738,7 @@ public:
 				uint32 ActualExtent = (*ArraySetPtr)->GetExtent();
 				if (ActualExtent >= TMeshAttributesRefType<T>::MinExpectedExtent && ActualExtent <= TMeshAttributesRefType<T>::MaxExpectedExtent)
 				{
-					return TMeshAttributesRef<int32, RefType>(static_cast<typename TMeshAttributesRefType<T>::SetType*>(ArraySetPtr->Get()), ActualExtent);
+					return TMeshAttributesRef<int32, RefType>(ArraySetPtr->Get(), ActualExtent);
 				}
 			}
 		}
@@ -1749,17 +1798,14 @@ public:
 		if (const FAttributesSetEntry* ArraySetPtr = this->Map.Find(AttributeName))
 		{
 			using AttributeType = typename TMeshAttributesRefType<T>::AttributeType;
-
-			// Get the real attribute type; this will normally be the same as the templated type,
-			// but element ID types are coerced to simple int32s.
-			using RealAttributeType = typename TChooseClass<TIsDerivedFrom<AttributeType, FElementID>::Value, int32, AttributeType>::Result;
+			using RealAttributeType = typename TMeshAttributesRefType<T>::RealAttributeType;
 
 			if ((*ArraySetPtr)->HasType<RealAttributeType>())
 			{
 				uint32 ActualExtent = (*ArraySetPtr)->GetExtent();
 				if (ActualExtent >= TMeshAttributesRefType<T>::MinExpectedExtent && ActualExtent <= TMeshAttributesRefType<T>::MaxExpectedExtent)
 				{
-					return TMeshAttributesConstRef<ElementIDType, RefType>(static_cast<typename TMeshAttributesRefType<T>::ConstSetType*>(ArraySetPtr->Get()), ActualExtent);
+					return TMeshAttributesConstRef<ElementIDType, RefType>(ArraySetPtr->Get(), ActualExtent);
 				}
 			}
 		}
@@ -1776,17 +1822,14 @@ public:
 		if (FAttributesSetEntry* ArraySetPtr = this->Map.Find(AttributeName))
 		{
 			using AttributeType = typename TMeshAttributesRefType<T>::AttributeType;
-
-			// Get the real attribute type; this will normally be the same as the templated type,
-			// but element ID types are coerced to simple int32s.
-			using RealAttributeType = typename TChooseClass<TIsDerivedFrom<AttributeType, FElementID>::Value, int32, AttributeType>::Result;
+			using RealAttributeType = typename TMeshAttributesRefType<T>::RealAttributeType;
 
 			if ((*ArraySetPtr)->HasType<RealAttributeType>())
 			{
 				uint32 ActualExtent = (*ArraySetPtr)->GetExtent();
 				if (ActualExtent >= TMeshAttributesRefType<T>::MinExpectedExtent && ActualExtent <= TMeshAttributesRefType<T>::MaxExpectedExtent)
 				{
-					return TMeshAttributesRef<ElementIDType, RefType>(static_cast<typename TMeshAttributesRefType<T>::SetType*>(ArraySetPtr->Get()), ActualExtent);
+					return TMeshAttributesRef<ElementIDType, RefType>(ArraySetPtr->Get(), ActualExtent);
 				}
 			}
 		}
@@ -1923,7 +1966,6 @@ public:
 	T GetAttribute(const ElementIDType ElementID, const FName AttributeName, const int32 AttributeChannel = 0) const
 	{
 		using RefType = typename TMeshAttributesRefType<T>::ConstRefType;
-		using SetType = typename TMeshAttributesRefType<T>::ConstSetType;
 		using AttributeType = typename TMeshAttributesRefType<T>::AttributeType;
 		using RealAttributeType = typename TMeshAttributesRefType<T>::RealAttributeType;
 
@@ -1932,7 +1974,7 @@ public:
 		check(ArraySetPtr->HasType<RealAttributeType>());
 		check(ActualExtent >= TMeshAttributesRefType<AttributeType>::MinExpectedExtent && ActualExtent <= TMeshAttributesRefType<AttributeType>::MaxExpectedExtent);
 
-		TMeshAttributesConstRef<FElementID, RefType> Ref(static_cast<SetType*>(ArraySetPtr), ActualExtent);
+		TMeshAttributesConstRef<FElementID, RefType> Ref(ArraySetPtr, ActualExtent);
 		return Ref.Get(ElementID, AttributeChannel);
 	}
 
@@ -1945,7 +1987,6 @@ public:
 	void SetAttribute(const ElementIDType ElementID, const FName AttributeName, const int32 AttributeChannel, const T& AttributeValue)
 	{
 		using NonConstRefType = typename TMeshAttributesRefType<T>::NonConstRefType;
-		using SetType = typename TMeshAttributesRefType<T>::SetType;
 		using AttributeType = typename TMeshAttributesRefType<T>::AttributeType;
 		using RealAttributeType = typename TMeshAttributesRefType<T>::RealAttributeType;
 
@@ -1954,7 +1995,7 @@ public:
 		check(ArraySetPtr->HasType<typename TRemoveCV<RealAttributeType>::Type>());
 		check(ActualExtent >= TMeshAttributesRefType<AttributeType>::MinExpectedExtent && ActualExtent <= TMeshAttributesRefType<AttributeType>::MaxExpectedExtent);
 
-		TMeshAttributesRef<FElementID, NonConstRefType> Ref(static_cast<SetType*>(ArraySetPtr), ActualExtent);
+		TMeshAttributesRef<FElementID, NonConstRefType> Ref(ArraySetPtr, ActualExtent);
 		return Ref.Set(ElementID, AttributeChannel, AttributeValue);
 	}
 
