Fix memory allocation in the TArrayAttribute

diff --git a/Engine/Source/Runtime/MeshDescription/Public/AttributeArrayContainer.h b/Engine/Source/Runtime/MeshDescription/Public/AttributeArrayContainer.h
--- a/Engine/Source/Runtime/MeshDescription/Public/AttributeArrayContainer.h
+++ b/Engine/Source/Runtime/MeshDescription/Public/AttributeArrayContainer.h
@@ -208,6 +208,7 @@ public:
 			// and adjust subsequent start indices accordingly.
 			int32 OldDataSize = Chunk.Data.Num();
 			Chunk.Data.SetNum(OldDataSize + Extra);
+			BasePtr = Chunk.Data.GetData();
 
 			AttributeType* StartPtr = BasePtr + Chunk.StartIndex[ElementIndex] + Chunk.MaxCount[ElementIndex];
 			AttributeType* EndPtr = BasePtr + OldDataSize;
