Changed TArrayAttribute to hold a pointer to the TAttributeArrayContainer instead of a reference. The end result is the same.

diff --git a/Engine/Source/Runtime/MeshDescription/Public/AttributeArrayContainer.h b/Engine/Source/Runtime/MeshDescription/Public/AttributeArrayContainer.h
--- a/Engine/Source/Runtime/MeshDescription/Public/AttributeArrayContainer.h
+++ b/Engine/Source/Runtime/MeshDescription/Public/AttributeArrayContainer.h
@@ -396,7 +396,7 @@ class TArrayAttribute
 
 public:
 	explicit TArrayAttribute(ArrayType& InArray, int32 InIndex)
-		: Array(InArray),
+		: Array(&InArray),
 		  Index(InIndex)
 	{}
 
@@ -414,7 +414,7 @@ public:
 	 */
 	AttributeType* GetData() const
 	{
-		return Array.Get(Index).GetData();
+		return Array->Get(Index).GetData();
 	}
 
 	/**
@@ -426,7 +426,7 @@ public:
 	 */
 	bool IsValidIndex(int32 ArrayAttributeIndex) const
 	{
-		return Array.Get(Index).IsValidIndex(ArrayAttributeIndex);
+		return Array->Get(Index).IsValidIndex(ArrayAttributeIndex);
 	}
 
 	/**
@@ -436,7 +436,7 @@ public:
 	 */
 	bool IsEmpty() const
 	{
-		return Array.Get(Index).IsEmpty();
+		return Array->Get(Index).IsEmpty();
 	}
 
 	/**
@@ -446,7 +446,7 @@ public:
 	 */
 	int32 Num() const
 	{
-		return Array.Get(Index).Num();
+		return Array->Get(Index).Num();
 	}
 
 	/**
@@ -456,7 +456,7 @@ public:
 	 */
 	AttributeType& operator[](int32 ArrayAttributeIndex) const
 	{
-		return Array.Get(Index)[ArrayAttributeIndex];
+		return Array->Get(Index)[ArrayAttributeIndex];
 	}
 
 	/**
@@ -469,7 +469,7 @@ public:
 	 */
 	AttributeType& Last(int32 IndexFromTheEnd = 0) const
 	{
-		return Array.Get(Index).Last();
+		return Array->Get(Index).Last();
 	}
 
 	/**
@@ -479,7 +479,7 @@ public:
 	{
 		const bool bSetCount = true;
 		const bool bSetDefault = true;
-		Array.SetElementSize(Index, Num, bSetCount, bSetDefault);
+		Array->SetElementSize(Index, Num, bSetCount, bSetDefault);
 	}
 
 
@@ -490,7 +490,7 @@ public:
 	{
 		const bool bSetCount = false;
 		const bool bSetDefault = true;
-		Array.SetElementSize(Index, Num, bSetCount, bSetDefault);
+		Array->SetElementSize(Index, Num, bSetCount, bSetDefault);
 	}
 
 
@@ -499,8 +499,8 @@ public:
 	 */
 	int32 Add(const AttributeType& Value)
 	{
-		int32 ElementCount = Array.GetElementCount(Index);
-		TArrayView<AttributeType> Element = Array.InsertIntoElement(Index, ElementCount, 1);
+		int32 ElementCount = Array->GetElementCount(Index);
+		TArrayView<AttributeType> Element = Array->InsertIntoElement(Index, ElementCount, 1);
 		Element[ElementCount] = Value;
 		return ElementCount;
 	}
@@ -528,7 +528,7 @@ public:
 	 */
 	void InsertDefaulted(int32 StartIndex, int32 Count)
 	{
-		Array.InsertIntoElement(Index, StartIndex, Count);
+		Array->InsertIntoElement(Index, StartIndex, Count);
 	}
 
 
@@ -537,7 +537,7 @@ public:
 	 */
 	void Insert(const AttributeType& Value, int32 StartIndex)
 	{
-		TArrayView<AttributeType> Element = Array.InsertIntoElement(Index, StartIndex, 1);
+		TArrayView<AttributeType> Element = Array->InsertIntoElement(Index, StartIndex, 1);
 		Element[StartIndex] = Value;
 	}
 
@@ -547,7 +547,7 @@ public:
 	 */
 	void RemoveAt(int32 StartIndex, int32 Count)
 	{
-		Array.RemoveFromElement(Index, StartIndex, Count);
+		Array->RemoveFromElement(Index, StartIndex, Count);
 	}
 
 
@@ -557,7 +557,7 @@ public:
 	template <typename Predicate>
 	int32 RemoveAll(Predicate Pred)
 	{
-		TArrayView<AttributeType> Element = Array.Get(Index);
+		TArrayView<AttributeType> Element = Array->Get(Index);
 		int32 Count = Element.Num();
 
 		int32 WriteIndex = 0;
@@ -572,7 +572,7 @@ public:
 
 		const bool bSetCount = true;
 		const bool bSetDefault = true;
-		Array.SetElementSize(Index, WriteIndex, bSetCount, bSetDefault);
+		Array->SetElementSize(Index, WriteIndex, bSetCount, bSetDefault);
 
 		return Count - WriteIndex;
 	}
@@ -589,7 +589,7 @@ public:
 	 */
 	int32 Find(const AttributeType& Value) const
 	{
-		return Array.Get(Index).Find(Value);
+		return Array->Get(Index).Find(Value);
 	}
 
 
@@ -604,7 +604,7 @@ public:
 	template <typename Predicate>
 	AttributeType* FindByPredicate(Predicate Pred) const
 	{
-		return Array.Get(Index).FindByPredicate(MoveTemp(Pred));
+		return Array->Get(Index).FindByPredicate(MoveTemp(Pred));
 	}
 
 
@@ -619,7 +619,7 @@ public:
 	template <typename Predicate>
 	int32 IndexOfByPredicate(Predicate Pred) const
 	{
-		return Array.Get(Index).IndexOfByPredicate(MoveTemp(Pred));
+		return Array->Get(Index).IndexOfByPredicate(MoveTemp(Pred));
 	}
 
 
@@ -628,7 +628,7 @@ public:
 	 */
 	bool Contains(const AttributeType& Value) const
 	{
-		return Array.Get(Index).Contains(Value);
+		return Array->Get(Index).Contains(Value);
 	}
 
 
@@ -642,7 +642,7 @@ public:
 	template <typename Predicate>
 	bool ContainsByPredicate(Predicate Pred) const
 	{
-		return Array.Get(Index).ContainsByPredicate(MoveTemp(Pred));
+		return Array->Get(Index).ContainsByPredicate(MoveTemp(Pred));
 	}
 
 
@@ -651,7 +651,7 @@ public:
 	 */
 	void Sort()
 	{
-		Array.Get(Index).Sort();
+		Array->Get(Index).Sort();
 	}
 
 
@@ -660,19 +660,19 @@ public:
 	 */
 	void StableSort()
 	{
-		Array.Get(Index).StableSort();
+		Array->Get(Index).StableSort();
 	}
 
 
 	/**
 	 * Return a TArrayView representing this array attribute.
 	 */
-	TArrayView<AttributeType> ToArrayView() { return Array.Get(Index); }
+	TArrayView<AttributeType> ToArrayView() { return Array->Get(Index); }
 
 	/**
 	 * Implicitly coerce an array attribute to a TArrayView.
 	 */
-	operator TArrayView<AttributeType>() { return Array.Get(Index); }
+	operator TArrayView<AttributeType>() { return Array->Get(Index); }
 
 public:
 	/**
@@ -683,7 +683,7 @@ public:
 	FORCEINLINE AttributeType* end  () const { return GetData() + Num(); }
 
 private:
-	ArrayType& Array;
+	ArrayType* Array;
 	int32 Index;
 };
 
