Changed behaviour of GetRawArray so that it returns an empty TArrayView when the container has zero elements. Previously it would attempt to dereference the underlying array allocation which may be nullptr for an empty array.

diff --git a/Engine/Source/Runtime/MeshDescription/Public/MeshAttributeArray.h b/Engine/Source/Runtime/MeshDescription/Public/MeshAttributeArray.h
--- a/Engine/Source/Runtime/MeshDescription/Public/MeshAttributeArray.h
+++ b/Engine/Source/Runtime/MeshDescription/Public/MeshAttributeArray.h
@@ -838,7 +838,7 @@ public:
 
 	TArrayView<AttributeType> GetRawArray(const int32 AttributeChannel = 0) const
 	{
-		if (ArrayPtr == nullptr)
+		if (ArrayPtr == nullptr || GetNumElements() == 0)
 		{
 			return TArrayView<AttributeType>();
 		}
@@ -1081,7 +1081,7 @@ public:
 		// Can't get the attribute set raw array for unbounded arrays because they are chunked
 		check(Extent > 0);
 
-		if (ArrayPtr == nullptr)
+		if (ArrayPtr == nullptr || GetNumElements() == 0)
 		{
 			return TArrayView<AttributeType>();
 		}
@@ -1302,7 +1302,7 @@ public:
 	/** In this specialization, GetRawArray returns a pointer to the attribute array container holding the attributes and their index pointers */
 	const TAttributeArrayContainer<AttributeType>* GetRawArray(const int32 AttributeChannel = 0) const
 	{
-		if (ArrayPtr == nullptr)
+		if (ArrayPtr == nullptr || GetNumElements() == 0)
 		{
 			return nullptr;
 		}
