MeshDescription: Add SortByPredicate, StableSortByPredicate and SetNumUninitialized to TArrayAttribute

#rnx

diff --git a/Engine/Source/Runtime/MeshDescription/Public/AttributeArrayContainer.h b/Engine/Source/Runtime/MeshDescription/Public/AttributeArrayContainer.h
--- a/Engine/Source/Runtime/MeshDescription/Public/AttributeArrayContainer.h
+++ b/Engine/Source/Runtime/MeshDescription/Public/AttributeArrayContainer.h
@@ -372,7 +372,7 @@ void TAttributeArrayContainer<AttributeType>::Remap(const TSparseArray<int32>& I
 {
 	TAttributeArrayContainer<AttributeType> NewAttributeArray(Default);
 
-	for (typename TSparseArray<int32>::TConstIterator It(IndexRemap); It; ++It)
+	for (TSparseArray<int32>::TConstIterator It(IndexRemap); It; ++It)
 	{
 		const int32 OldElementIndex = It.GetIndex();
 		const int32 NewElementIndex = IndexRemap[OldElementIndex];
@@ -483,6 +483,19 @@ public:
 	}
 
 
+	/**
+	 * Sets the number of elements in the array attribute without initializing
+	 * any new elements. It is the responsibility of the caller to do so after this 
+	 * function returns.
+	 */
+	void SetNumUninitialized(int32 Num)
+	{
+		const bool bSetCount = true;
+		const bool bSetDefault = false;
+		Array->SetElementSize(Index, Num, bSetCount, bSetDefault);
+	}
+
+
 	/**
 	 * Reserves the number of elements in the array attribute.
 	 */
@@ -654,6 +667,15 @@ public:
 		Array->Get(Index).Sort();
 	}
 
+	/**
+	 * Sorts the array using the given predicate.
+	 */
+	template<typename Predicate>
+	void SortByPredicate(Predicate Pred)
+	{
+		Array->Get(Index).template Sort<Predicate>(Pred);
+	}
+
 
 	/**
 	 * Stable sorts the array assuming < operator is defined for the item type.
@@ -664,6 +686,15 @@ public:
 	}
 
 
+	/**
+	 * Sorts the array using the given predicate.
+	 */
+	template<typename Predicate>
+	void StableSortByPredicate(Predicate Pred)
+	{
+		Array->Get(Index).template StableSort<Predicate>(Pred);
+	}
+
 	/**
 	 * Return a TArrayView representing this array attribute.
 	 */
