MeshDescription fixes
* Added missing write lock in FMeshDescriptionBulkData::SaveMeshDescription( FMeshDescription& MeshDescription )
* Skip serialization of bulk data when archive is a ref collector

diff --git a/Engine/Source/Runtime/MeshDescription/Private/MeshDescription.cpp b/Engine/Source/Runtime/MeshDescription/Private/MeshDescription.cpp
--- a/Engine/Source/Runtime/MeshDescription/Private/MeshDescription.cpp
+++ b/Engine/Source/Runtime/MeshDescription/Private/MeshDescription.cpp
@@ -1964,6 +1964,12 @@ void FMeshDescriptionBulkData::Serialize( FArchive& Ar, UObject* Owner )
 	Ar.UsingCustomVersion(FUE5CookerObjectVersion::GUID);
 #endif //UE_USE_VIRTUALBULKDATA
 
+	// Make sure to serialize only actual data
+	if (Ar.ShouldSkipBulkData() || Ar.IsObjectReferenceCollector() || !Ar.IsPersistent())
+	{
+		return;
+	}
+
 	if( Ar.IsTransacting() )
 	{
 		// If transacting, keep these members alive the other side of an undo, otherwise their values will get lost
@@ -2039,6 +2045,10 @@ void FMeshDescriptionBulkData::SaveMeshDescription( FMeshDescription& MeshDescri
 {
 	TRACE_CPUPROFILER_EVENT_SCOPE(FMeshDescriptionBulkData::SaveMeshDescription);
 
+#if WITH_EDITOR
+	FRWScopeLock ScopeLock(BulkDataLock, SLT_Write);
+#endif
+
 #if UE_USE_VIRTUALBULKDATA
 	BulkData.Reset();
 #else
