Re-fixed undo/redo of static mesh geometry (FMeshDescriptionBulkData serialization). This must be allowed to happen for non-persistent archives, such as the transaction buffer.

diff --git a/Engine/Source/Runtime/MeshDescription/Private/MeshDescription.cpp b/Engine/Source/Runtime/MeshDescription/Private/MeshDescription.cpp
--- a/Engine/Source/Runtime/MeshDescription/Private/MeshDescription.cpp
+++ b/Engine/Source/Runtime/MeshDescription/Private/MeshDescription.cpp
@@ -1965,7 +1965,7 @@ void FMeshDescriptionBulkData::Serialize( FArchive& Ar, UObject* Owner )
 #endif //UE_USE_VIRTUALBULKDATA
 
 	// Make sure to serialize only actual data
-	if (Ar.ShouldSkipBulkData() || Ar.IsObjectReferenceCollector() || !Ar.IsPersistent())
+	if (Ar.ShouldSkipBulkData() || Ar.IsObjectReferenceCollector())
 	{
 		return;
 	}
