The MeshDescription object within the bulkdata is no longer marked as Transactable. All transactable changes must occur on the bulkdata itself, and these changes will be propagated to the cached MeshDescription, if one exists, during an undo/redo. Corrected BulkData serialize by only caching the MeshDescription from the BulkData if the archive is loading.

diff --git a/Engine/Source/Runtime/MeshDescription/Private/MeshDescriptionBaseBulkData.cpp b/Engine/Source/Runtime/MeshDescription/Private/MeshDescriptionBaseBulkData.cpp
--- a/Engine/Source/Runtime/MeshDescription/Private/MeshDescriptionBaseBulkData.cpp
+++ b/Engine/Source/Runtime/MeshDescription/Private/MeshDescriptionBaseBulkData.cpp
@@ -25,7 +25,7 @@ void UMeshDescriptionBaseBulkData::Serialize(FArchive& Ar)
 		BulkData->Serialize(Ar, GetOuter());
 	}
 
-	if (MeshDescription != nullptr)
+	if (Ar.IsLoading() && MeshDescription != nullptr)
 	{
 		// If there was a cached mesh description, it could be out of sync with the bulk data we just serialized so re-cache it.
 		CacheMeshDescription();
@@ -79,7 +79,7 @@ UMeshDescriptionBase* UMeshDescriptionBaseBulkData::CreateMeshDescription()
 		// Assign the new UObject directly to a UPROPERTY.
 		// If we were run from a different thread, we can immediately clear the Async internal flag, now that GC refs are set up correctly.
 		UClass* MeshDescriptionType = GetMeshDescriptionType();
-		MeshDescription = NewObject<UMeshDescriptionBase>(this, MeshDescriptionType, NAME_None, RF_Transient | RF_Transactional);
+		MeshDescription = NewObject<UMeshDescriptionBase>(this, MeshDescriptionType, NAME_None, RF_Transient);
 		MeshDescription->AtomicallyClearInternalFlags(EInternalObjectFlags::Async);
 	}
 
