Correctly construct an instanced UMeshDescriptionBase as an inner object of UMeshDescriptionBaseBulkData when duplicating a static mesh.

diff --git a/Engine/Source/Runtime/MeshDescription/Public/MeshDescriptionBase.h b/Engine/Source/Runtime/MeshDescription/Public/MeshDescriptionBase.h
--- a/Engine/Source/Runtime/MeshDescription/Public/MeshDescriptionBase.h
+++ b/Engine/Source/Runtime/MeshDescription/Public/MeshDescriptionBase.h
@@ -8,7 +8,7 @@
 
 
 
-UCLASS(BlueprintType)
+UCLASS(BlueprintType, DefaultToInstanced)
 class MESHDESCRIPTION_API UMeshDescriptionBase : public UObject
 {
 	GENERATED_BODY()
