Better way to ensure that the inner UMeshDescriptionBase is instanced when duplicating the bulkdata wrapper.

diff --git a/Engine/Source/Runtime/MeshDescription/Public/MeshDescriptionBase.h b/Engine/Source/Runtime/MeshDescription/Public/MeshDescriptionBase.h
--- a/Engine/Source/Runtime/MeshDescription/Public/MeshDescriptionBase.h
+++ b/Engine/Source/Runtime/MeshDescription/Public/MeshDescriptionBase.h
@@ -8,7 +8,7 @@
 
 
 
-UCLASS(BlueprintType, DefaultToInstanced)
+UCLASS(BlueprintType)
 class MESHDESCRIPTION_API UMeshDescriptionBase : public UObject
 {
 	GENERATED_BODY()
diff --git a/Engine/Source/Runtime/MeshDescription/Public/MeshDescriptionBaseBulkData.h b/Engine/Source/Runtime/MeshDescription/Public/MeshDescriptionBaseBulkData.h
--- a/Engine/Source/Runtime/MeshDescription/Public/MeshDescriptionBaseBulkData.h
+++ b/Engine/Source/Runtime/MeshDescription/Public/MeshDescriptionBaseBulkData.h
@@ -47,7 +47,7 @@ public:
 protected:
 	TPimplPtr<FMeshDescriptionBulkData> BulkData;
 
-	UPROPERTY(Transient)
+	UPROPERTY(Transient, Instanced)
 	UMeshDescriptionBase* PreallocatedMeshDescription;
 
 	UPROPERTY(Transient)
