Added a TAttributesSet::ForEachByType utility function to iterate over attributes of a specific type (e.g. all floats). Allows better specialization of the called function for a given type.

diff --git a/Engine/Source/Runtime/MeshDescription/Public/MeshAttributeArray.h b/Engine/Source/Runtime/MeshDescription/Public/MeshAttributeArray.h
--- a/Engine/Source/Runtime/MeshDescription/Public/MeshAttributeArray.h
+++ b/Engine/Source/Runtime/MeshDescription/Public/MeshAttributeArray.h
@@ -2024,6 +2024,20 @@ public:
 	* The prototype should be Func( const FName AttributeName, auto AttributesConstRef );
 	*/
 	template <typename ForEachFunc> void ForEach(ForEachFunc Func) const;
+
+	/**
+	* Call the supplied function on each attribute that matches the given type.
+	* The type can be given as either a plain type T, TArrayView<T> or TArrayAttribute<T>. 
+	* The prototype should be Func( const FName AttributeName, auto AttributesRef );
+	*/
+	template <typename AttributeType, typename ForEachFunc> void ForEachByType(ForEachFunc Func);
+	
+	/**
+	* Call the supplied function on each attribute that matches the given type.
+	* The type can be given as either a plain type T, TArrayView<T> or TArrayAttribute<T>. 
+	* The prototype should be Func( const FName AttributeName, auto AttributesConstRef );
+	*/
+	template <typename AttributeType, typename ForEachFunc> void ForEachByType(ForEachFunc Func) const;
 };
 
 
@@ -2124,6 +2138,7 @@ void TAttributesSet<ElementIDType>::ForEach(ForEachFunc Func)
 	}
 }
 
+
 namespace ForEachConstImpl
 {
 	// Declare type of jump table used to dispatch functions
@@ -2172,6 +2187,102 @@ void TAttributesSet<ElementIDType>::ForEach(ForEachFunc Func) const
 	}
 }
 
+
+namespace ForEachByTypeImpl
+{
+	template<typename ElementIDType, typename AttributeType, typename ForEachFunc>
+    struct DispatchFunctor
+	{
+		void operator()(FName Name, ForEachFunc Fn, FMeshAttributeArraySetBase* Attributes)
+		{
+			if (TTupleIndex<AttributeType, AttributeTypes>::Value == Attributes->GetType() && Attributes->GetExtent() == 1)
+			{
+				Fn(Name, TMeshAttributesRef<ElementIDType, AttributeType>(static_cast<TMeshAttributeArraySet<AttributeType>*>(Attributes)));
+			}
+		}
+	};
+
+	template<typename ElementIDType,  typename AttributeType, typename ForEachFunc>
+    struct DispatchFunctor<ElementIDType, TArrayView<AttributeType>, ForEachFunc>
+	{
+		void operator()(FName Name, ForEachFunc Fn, FMeshAttributeArraySetBase* Attributes)
+		{
+			if (TTupleIndex<AttributeType, AttributeTypes>::Value == Attributes->GetType() && Attributes->GetExtent() >= 1)
+			{
+				Fn(Name, TMeshAttributesRef<ElementIDType, TArrayView<AttributeType>>(static_cast<TMeshAttributeArraySet<AttributeType>*>(Attributes), Attributes->GetExtent()));
+			}
+		}
+	};
+
+	template<typename ElementIDType, typename AttributeType, typename ForEachFunc>
+    struct DispatchFunctor<ElementIDType, TArrayAttribute<AttributeType>, ForEachFunc>
+	{
+		void operator()(FName Name, ForEachFunc Fn, FMeshAttributeArraySetBase* Attributes)
+		{
+			if (TTupleIndex<AttributeType, AttributeTypes>::Value == Attributes->GetType() && Attributes->GetExtent() == 0)
+			{
+				Fn(Name, TMeshAttributesRef<ElementIDType, TArrayAttribute<AttributeType>>(static_cast<TMeshUnboundedAttributeArraySet<AttributeType>*>(Attributes)));
+			}
+		}
+	};
+	
+	template<typename ElementIDType, typename AttributeType, typename ForEachFunc>
+    struct ConstDispatchFunctor
+	{
+		void operator()(FName Name, ForEachFunc Fn, const FMeshAttributeArraySetBase* Attributes)
+		{
+			if (TTupleIndex<AttributeType, AttributeTypes>::Value == Attributes->GetType() && Attributes->GetExtent() == 1)
+			{
+				Fn(Name, TMeshAttributesConstRef<ElementIDType, AttributeType>(static_cast<const TMeshAttributeArraySet<AttributeType>*>(Attributes)));
+			}
+		}
+	};
+
+	template<typename ElementIDType,  typename AttributeType, typename ForEachFunc>
+    struct ConstDispatchFunctor<ElementIDType, TArrayView<AttributeType>, ForEachFunc>
+	{
+		void operator()(FName Name, ForEachFunc Fn, const FMeshAttributeArraySetBase* Attributes)
+		{
+			if (TTupleIndex<AttributeType, AttributeTypes>::Value == Attributes->GetType() && Attributes->GetExtent() >= 1)
+			{
+				Fn(Name, TMeshAttributesConstRef<ElementIDType, TArrayView<const AttributeType>>(static_cast<const TMeshAttributeArraySet<AttributeType>*>(Attributes), Attributes->GetExtent()));
+			}
+		}
+	};
+
+	template<typename ElementIDType, typename AttributeType, typename ForEachFunc>
+    struct ConstDispatchFunctor<ElementIDType, TArrayAttribute<AttributeType>, ForEachFunc>
+	{
+		void operator()(FName Name, ForEachFunc Fn, const FMeshAttributeArraySetBase* Attributes)
+		{
+			if (TTupleIndex<AttributeType, AttributeTypes>::Value == Attributes->GetType() && Attributes->GetExtent() == 0)
+			{
+				Fn(Name, TMeshAttributesConstRef<ElementIDType, TArrayAttribute<const AttributeType>>(static_cast<const TMeshUnboundedAttributeArraySet<AttributeType>*>(Attributes)));
+			}
+		}
+	};}
+
+template <typename ElementIDType>
+template < typename AttributeType, typename ForEachFunc>
+void TAttributesSet<ElementIDType>::ForEachByType(ForEachFunc Func)
+{
+	for (auto& MapEntry : this->Map)
+	{
+		ForEachByTypeImpl::DispatchFunctor<ElementIDType, typename TRemoveConst<AttributeType>::Type, ForEachFunc>()(MapEntry.Key, Func, MapEntry.Value.Get());
+	}
+}
+
+template <typename ElementIDType>
+template < typename AttributeType, typename ForEachFunc>
+void TAttributesSet<ElementIDType>::ForEachByType(ForEachFunc Func) const
+{
+	for (const auto& MapEntry : this->Map)
+	{
+		ForEachByTypeImpl::ConstDispatchFunctor<ElementIDType, typename TRemoveConst<AttributeType>::Type, ForEachFunc>()(MapEntry.Key, Func, MapEntry.Value.Get());
+	}
+}
+
+
 /**
  * This is a similar approach to ForEach, above.
  * Given a type index, at runtime, we wish to create an attribute array of the corresponding type; essentially a factory.
