Avoid renderdata of HLOD being recomputed during cook package save because of non-deterministic guid
  - Preserve guid during save if static mesh bulkdata content hasn't changed to avoid triggering unnecessary rebuilds
  - Make sure the next HLOD generation is going to use the HashAsGuid feature to avoid recompiling staticmesh that are generated if they are exactly the same as before

  - 31m10s to 14m15s to cook P_Construct_WP with a warmed-up DDC

diff --git a/Engine/Source/Developer/MeshMergeUtilities/Private/ProxyGenerationProcessor.cpp b/Engine/Source/Developer/MeshMergeUtilities/Private/ProxyGenerationProcessor.cpp
--- a/Engine/Source/Developer/MeshMergeUtilities/Private/ProxyGenerationProcessor.cpp
+++ b/Engine/Source/Developer/MeshMergeUtilities/Private/ProxyGenerationProcessor.cpp
@@ -249,7 +249,9 @@ void FProxyGenerationProcessor::ProcessJob(const FGuid& JobGuid, FProxyGeneratio
 				StaticMesh->GetStaticMaterials().Add(NewMaterial);
 			}
 
-			StaticMesh->CommitMeshDescription(SourceModelIndex);
+			UStaticMesh::FCommitMeshDescriptionParams CommitParams;
+			CommitParams.bUseHashAsGuid = true;
+			StaticMesh->CommitMeshDescription(SourceModelIndex, CommitParams);
 		}
 	};
 
diff --git a/Engine/Source/Runtime/MeshDescription/Private/MeshDescription.cpp b/Engine/Source/Runtime/MeshDescription/Private/MeshDescription.cpp
--- a/Engine/Source/Runtime/MeshDescription/Private/MeshDescription.cpp
+++ b/Engine/Source/Runtime/MeshDescription/Private/MeshDescription.cpp
@@ -1984,9 +1984,23 @@ void FMeshDescriptionBulkData::Serialize( FArchive& Ar, UObject* Owner )
 			// Explicitly load and resave the FMeshDescription so that its version is in sync with the FMeshDescriptionBulkData.
 			if( !bBulkDataUpdated )
 			{
+				const FGuid OriginalGuid = Guid;
+				FGuid OriginalHash;
+				if (!bGuidIsHash)
+				{
+					OriginalHash = GetHash();
+				}
+
 				FMeshDescription MeshDescription;
 				LoadMeshDescription( MeshDescription );
 				SaveMeshDescription( MeshDescription );
+
+				// Maintain original guid in case the hash is the same to avoid invalidating
+				// all the data already built using this guid.
+				if (!bGuidIsHash && OriginalHash == GetHash())
+				{
+					Guid = OriginalGuid;
+				}
 			}
 		}
 	}
@@ -2156,15 +2170,12 @@ FString FMeshDescriptionBulkData::GetIdString() const
 	return GuidString;
 }
 
-
-void FMeshDescriptionBulkData::UseHashAsGuid()
+FGuid FMeshDescriptionBulkData::GetHash() const
 {
 	if (BulkData.GetBulkDataSize() > 0)
 	{
-		bGuidIsHash = true;
-
 #if UE_USE_VIRTUALBULKDATA
-		Guid = BulkData.GetKey();
+		return BulkData.GetKey();
 #else
 		uint32 Hash[5] = {};
 
@@ -2172,9 +2183,21 @@ void FMeshDescriptionBulkData::UseHashAsGuid()
 		FSHA1::HashBuffer(Buffer, BulkData.GetBulkDataSize(), (uint8*)Hash);
 		BulkData.Unlock();
 
-		Guid = FGuid(Hash[0] ^ Hash[4], Hash[1], Hash[2], Hash[3]);
+		return FGuid(Hash[0] ^ Hash[4], Hash[1], Hash[2], Hash[3]);
 #endif //UE_USE_VIRTUALBULKDATA
 	}
+
+	return FGuid();
+}
+
+void FMeshDescriptionBulkData::UseHashAsGuid()
+{
+	if (BulkData.GetBulkDataSize() > 0)
+	{
+		bGuidIsHash = true;
+
+		Guid = GetHash();
+	}
 	else
 	{
 		Guid.Invalidate();
diff --git a/Engine/Source/Runtime/MeshDescription/Public/MeshDescription.h b/Engine/Source/Runtime/MeshDescription/Public/MeshDescription.h
--- a/Engine/Source/Runtime/MeshDescription/Public/MeshDescription.h
+++ b/Engine/Source/Runtime/MeshDescription/Public/MeshDescription.h
@@ -1329,6 +1329,9 @@ private:
 	/** GUID associated with the data stored herein. */
 	FGuid Guid;
 
+	/** Gets this bulk data hash */
+	FGuid GetHash() const;
+
 	/** Take a copy of the bulk data versioning so it can be propagated to the bulk data reader when deserializing MeshDescription */
 	FCustomVersionContainer CustomVersions;
 
