Ensure certain legacy meshes contain a UV element container.

diff --git a/Engine/Source/Runtime/MeshDescription/Private/MeshDescription.cpp b/Engine/Source/Runtime/MeshDescription/Private/MeshDescription.cpp
--- a/Engine/Source/Runtime/MeshDescription/Private/MeshDescription.cpp
+++ b/Engine/Source/Runtime/MeshDescription/Private/MeshDescription.cpp
@@ -192,6 +192,13 @@ void FMeshDescription::Serialize(FArchive& BaseAr)
 		// After loading elements, we need to re-cache mesh element and attribute arrays 
 		if (Ar.IsLoading())
 		{
+			// Ensure there's a UV element container
+			if (Elements.Find(UVsName) == nullptr)
+			{
+				UE_LOG(LogLoad, Warning, TEXT("Couldn't find UV element container in mesh - adding an empty one"));
+				UVElements = Elements.Emplace(UVsName).Get();
+			}
+
 			Cache();
 
 			for (FTriangleID TriangleID : TriangleElements->Get().GetElementIDs())
