MeshDescription: Eliminate unnecessary copy of index variables in for loop that triggers Clang's -Wrange-loop-analysis compile error.
#rnx

diff --git a/Engine/Source/Runtime/MeshDescription/Public/MeshDescription.h b/Engine/Source/Runtime/MeshDescription/Public/MeshDescription.h
--- a/Engine/Source/Runtime/MeshDescription/Public/MeshDescription.h
+++ b/Engine/Source/Runtime/MeshDescription/Public/MeshDescription.h
@@ -460,7 +460,7 @@ public:
 	void GetVertexConnectedTriangles(const FVertexID VertexID, TArray<FTriangleID, Alloc>& OutConnectedTriangleIDs) const
 	{
 		OutConnectedTriangleIDs.Reset(GetNumVertexConnectedTriangles(VertexID));
-		for (const FVertexInstanceID VertexInstanceID : VertexToVertexInstances.Find<FVertexInstanceID>(VertexID))
+		for (const FVertexInstanceID& VertexInstanceID : VertexToVertexInstances.Find<FVertexInstanceID>(VertexID))
 		{
 			TArrayView<const FTriangleID> ConnectedTris = VertexInstanceToTriangles.Find<FTriangleID>(VertexInstanceID);
 			OutConnectedTriangleIDs.Append(ConnectedTris.GetData(), ConnectedTris.Num());
@@ -498,9 +498,9 @@ public:
 	void GetVertexConnectedPolygons(const FVertexID VertexID, TArray<FPolygonID, Alloc>& OutConnectedPolygonIDs) const
 	{
 		OutConnectedPolygonIDs.Reset();
-		for (const FVertexInstanceID VertexInstanceID : VertexToVertexInstances.Find<FVertexInstanceID>(VertexID))
+		for (const FVertexInstanceID& VertexInstanceID : VertexToVertexInstances.Find<FVertexInstanceID>(VertexID))
 		{
-			for (const FTriangleID TriangleID : VertexInstanceToTriangles.Find<FTriangleID>(VertexInstanceID))
+			for (const FTriangleID& TriangleID : VertexInstanceToTriangles.Find<FTriangleID>(VertexInstanceID))
 			{
 				OutConnectedPolygonIDs.AddUnique(TrianglePolygons[TriangleID]);
 			}
@@ -538,7 +538,7 @@ public:
 		OutAdjacentVertexIDs.SetNumUninitialized(ConnectedEdgeIDs.Num());
 
 		int32 Index = 0;
-		for (const FEdgeID EdgeID : ConnectedEdgeIDs)
+		for (const FEdgeID& EdgeID : ConnectedEdgeIDs)
 		{
 			TArrayView<const FVertexID> EdgeVertexIDs = EdgeVertices[EdgeID];
 			OutAdjacentVertexIDs[Index] = (EdgeVertexIDs[0] == VertexID) ? EdgeVertexIDs[1] : EdgeVertexIDs[0];
@@ -614,7 +614,7 @@ public:
 	void GetVertexInstanceConnectedPolygons(const FVertexInstanceID VertexInstanceID, TArray<FPolygonID, Alloc>& OutPolygonIDs) const
 	{
 		OutPolygonIDs.Reset(VertexInstanceToTriangles.Find(VertexInstanceID).Num());
-		for (const FTriangleID TriangleID : VertexInstanceToTriangles.Find<FTriangleID>(VertexInstanceID))
+		for (const FTriangleID& TriangleID : VertexInstanceToTriangles.Find<FTriangleID>(VertexInstanceID))
 		{
 			OutPolygonIDs.AddUnique(TrianglePolygons[TriangleID]);
 		}
@@ -686,7 +686,7 @@ public:
 	void GetEdgeConnectedPolygons(const FEdgeID EdgeID, TArray<FPolygonID, Alloc>& OutPolygonIDs) const
 	{
 		OutPolygonIDs.Reset(EdgeToTriangles.Find(EdgeID).Num());
-		for (const FTriangleID TriangleID : EdgeToTriangles.Find<FTriangleID>(EdgeID))
+		for (const FTriangleID& TriangleID : EdgeToTriangles.Find<FTriangleID>(EdgeID))
 		{
 			OutPolygonIDs.AddUnique(TrianglePolygons[TriangleID]);
 		}
@@ -806,10 +806,11 @@ public:
 	void GetTriangleAdjacentTriangles(const FTriangleID TriangleID, TArray<FTriangleID, Alloc>& OutTriangleIDs) const
 	{
 		OutTriangleIDs.Reset();
-		for (const FEdgeID EdgeID : GetTriangleEdges(TriangleID))
+		for (const FEdgeID& EdgeID : GetTriangleEdges(TriangleID))
 		{
-			for (const FTriangleID OtherTriangleID : EdgeToTriangles.Find(EdgeID))
+			for (const int32& TriangleIndex : EdgeToTriangles.Find(EdgeID))
 			{
+				FTriangleID OtherTriangleID(TriangleIndex);
 				if (OtherTriangleID != TriangleID)
 				{
 					OutTriangleIDs.Add(OtherTriangleID);
@@ -1015,9 +1016,9 @@ public:
 		if (GetNumPolygonVertices(PolygonID) > 3)
 		{
 			TArray<FVertexInstanceID> VertexInstanceIDs = GetPolygonVertexInstances(PolygonID);
-			for (const FVertexInstanceID VertexInstanceID : VertexInstanceIDs)
+			for (const FVertexInstanceID& VertexInstanceID : VertexInstanceIDs)
 			{
-				for (const FEdgeID EdgeID : GetVertexConnectedEdgeIDs(GetVertexInstanceVertex(VertexInstanceID)))
+				for (const FEdgeID& EdgeID : GetVertexConnectedEdgeIDs(GetVertexInstanceVertex(VertexInstanceID)))
 				{
 					if (!OutEdgeIDs.Contains(EdgeID) && IsEdgeInternalToPolygon(EdgeID, PolygonID))
 					{
@@ -1056,9 +1057,9 @@ public:
 	void GetPolygonAdjacentPolygons(const FPolygonID PolygonID, TArray<FPolygonID, Alloc>& OutPolygonIDs) const
 	{
 		OutPolygonIDs.Reset();
-		for (const FEdgeID EdgeID : GetPolygonPerimeterEdges<TInlineAllocator<16>>(PolygonID))
+		for (const FEdgeID& EdgeID : GetPolygonPerimeterEdges<TInlineAllocator<16>>(PolygonID))
 		{
-			for (const FPolygonID OtherPolygonID : GetEdgeConnectedPolygons<TInlineAllocator<8>>(EdgeID))
+			for (const FPolygonID& OtherPolygonID : GetEdgeConnectedPolygons<TInlineAllocator<8>>(EdgeID))
 			{
 				if (OtherPolygonID != PolygonID)
 				{
@@ -1114,8 +1115,9 @@ public:
 		PolygonPolygonGroups[PolygonID] = PolygonGroupID;
 		PolygonGroupToPolygons.AddReferenceToKey(PolygonGroupID, PolygonID);
 
-		for (const FTriangleID TriangleID : PolygonToTriangles.Find(PolygonID))
+		for (const int32& TriangleIndex: PolygonToTriangles.Find(PolygonID))
 		{
+			const FTriangleID TriangleID(TriangleIndex);
 			check(TrianglePolygonGroups[TriangleID] == OldPolygonGroupID);
 			PolygonGroupToTriangles.RemoveReferenceFromKey(OldPolygonGroupID, TriangleID);
 
