MeshDescription: Fix Guid being changed during UpdateMeshDescriptionFormat. UpdateMeshDescriptionFormat was created to avoid changing the guid, but a copy-paste error left in the change of the guid.
#rnx

diff --git a/Engine/Source/Runtime/MeshDescription/Private/MeshDescription.cpp b/Engine/Source/Runtime/MeshDescription/Private/MeshDescription.cpp
--- a/Engine/Source/Runtime/MeshDescription/Private/MeshDescription.cpp
+++ b/Engine/Source/Runtime/MeshDescription/Private/MeshDescription.cpp
@@ -2176,7 +2176,6 @@ void FMeshDescriptionBulkData::UpdateMeshDescriptionFormat()
 		else
 		{
 			// Maintain the original guid; this is a format change only. GetIdString will change because it adds the CustomVersions to the key.
-			FPlatformMisc::CreateGuid(Guid);
 		}
 	}
 
