MeshDescription: underlying BitArray is not freeing memory on Empty() call as intended, use Reset() instead
#rnx

diff --git a/Engine/Source/Runtime/MeshDescription/Private/MeshDescriptionBaseBulkData.cpp b/Engine/Source/Runtime/MeshDescription/Private/MeshDescriptionBaseBulkData.cpp
--- a/Engine/Source/Runtime/MeshDescription/Private/MeshDescriptionBaseBulkData.cpp
+++ b/Engine/Source/Runtime/MeshDescription/Private/MeshDescriptionBaseBulkData.cpp
@@ -134,7 +134,7 @@ void UMeshDescriptionBaseBulkData::RemoveMeshDescription()
 {
 	if (MeshDescription != nullptr)
 	{
-		MeshDescription->Empty();
+		MeshDescription->Reset();
 		MeshDescription = nullptr;
 	}
 }
