Added access to Elements map in MeshDescription.

diff --git a/Engine/Source/Runtime/MeshDescription/Public/MeshDescription.h b/Engine/Source/Runtime/MeshDescription/Public/MeshDescription.h
--- a/Engine/Source/Runtime/MeshDescription/Public/MeshDescription.h
+++ b/Engine/Source/Runtime/MeshDescription/Public/MeshDescription.h
@@ -168,6 +168,9 @@ public:
 	TAttributesSet<FPolygonGroupID>& PolygonGroupAttributes() { return PolygonGroups().GetAttributes(); }
 	const TAttributesSet<FPolygonGroupID>& PolygonGroupAttributes() const { return PolygonGroups().GetAttributes(); }
 
+	TMap<FName, FMeshElementTypeWrapper>& GetElements() { return Elements; }
+	const TMap<FName, FMeshElementTypeWrapper>& GetElements() const { return Elements; }
+
 	void SuspendVertexIndexing() { VertexToVertexInstances.Suspend(); VertexToEdges.Suspend(); }
 	void SuspendVertexInstanceIndexing() { VertexInstanceToTriangles.Suspend(); }
 	void SuspendEdgeIndexing() { EdgeToTriangles.Suspend(); }
@@ -1264,6 +1267,7 @@ private:
 		Container.Add(Item);
 	}
 
+public:
 	static FName VerticesName;
 	static FName VertexInstancesName;
 	static FName UVsName;
@@ -1272,6 +1276,7 @@ private:
 	static FName PolygonsName;
 	static FName PolygonGroupsName;
 
+private:
 	TMap<FName, FMeshElementTypeWrapper> Elements;
 
 	FMeshElementChannels* VertexElements;
