Export FMeshElementContainer non-inline methods.

diff --git a/Engine/Source/Runtime/MeshDescription/Public/MeshElementContainer.h b/Engine/Source/Runtime/MeshDescription/Public/MeshElementContainer.h
--- a/Engine/Source/Runtime/MeshDescription/Public/MeshElementContainer.h
+++ b/Engine/Source/Runtime/MeshDescription/Public/MeshElementContainer.h
@@ -139,10 +139,10 @@ public:
 	FORCEINLINE const FAttributesSetBase& GetAttributes() const { return Attributes; }
 
 	/** Compacts elements and returns a remapping table */
-	void Compact(TSparseArray<int32>& OutIndexRemap);
+	MESHDESCRIPTION_API void Compact(TSparseArray<int32>& OutIndexRemap);
 
 	/** Remaps elements according to the passed remapping table */
-	void Remap(const TSparseArray<int32>& IndexRemap);
+	MESHDESCRIPTION_API void Remap(const TSparseArray<int32>& IndexRemap);
 
 	/** Serializer */
 	friend FArchive& operator<<(FArchive& Ar, FMeshElementContainer& Container)
