MeshDescription: In addition to applying the package's CustomVersions when deserializing MeshDescription out of BulkData, also apply the UEVersion and LicenseeUEVersion.
#rnx

diff --git a/Engine/Source/Runtime/CoreUObject/Private/Serialization/EditorBulkData.cpp b/Engine/Source/Runtime/CoreUObject/Private/Serialization/EditorBulkData.cpp
--- a/Engine/Source/Runtime/CoreUObject/Private/Serialization/EditorBulkData.cpp
+++ b/Engine/Source/Runtime/CoreUObject/Private/Serialization/EditorBulkData.cpp
@@ -16,6 +16,7 @@
 #include "UObject/LinkerSave.h"
 #include "UObject/Object.h"
 #include "UObject/ObjectSaveContext.h"
+#include "UObject/ObjectVersion.h"
 #include "UObject/PropertyPortFlags.h"
 #include "UObject/UObjectGlobals.h"
 #include "UObject/SavePackage.h"
@@ -1917,7 +1918,19 @@ void FEditorBulkData::SetCompressionOptions(ECompressedBufferCompressor Compress
 
 FCustomVersionContainer FEditorBulkData::GetCustomVersions(FArchive& InlineArchive)
 {
-	return InlineArchive.GetCustomVersions();
+	FPackageFileVersion OutUEVersion;
+	int32 OutLicenseeUEVersion;
+	FCustomVersionContainer OutCustomVersions;
+	GetBulkDataVersions(InlineArchive, OutUEVersion, OutLicenseeUEVersion, OutCustomVersions);
+	return OutCustomVersions;
+}
+
+void FEditorBulkData::GetBulkDataVersions(FArchive& InlineArchive, FPackageFileVersion& OutUEVersion,
+	int32& OutLicenseeUEVersion, FCustomVersionContainer& OutCustomVersions) const
+{
+	OutUEVersion = InlineArchive.UEVer();
+	OutLicenseeUEVersion = InlineArchive.LicenseeUEVer();
+	OutCustomVersions = InlineArchive.GetCustomVersions();
 }
 
 void FEditorBulkData::UpdatePayloadId()
diff --git a/Engine/Source/Runtime/CoreUObject/Public/Serialization/EditorBulkData.h b/Engine/Source/Runtime/CoreUObject/Public/Serialization/EditorBulkData.h
--- a/Engine/Source/Runtime/CoreUObject/Public/Serialization/EditorBulkData.h
+++ b/Engine/Source/Runtime/CoreUObject/Public/Serialization/EditorBulkData.h
@@ -21,6 +21,7 @@ class FArchive;
 class FBulkData;
 class FLinkerSave;
 class UObject;
+struct FPackageFileVersion;
 
 namespace UE::BulkDataRegistry { enum class ERegisterResult : uint8; }
 
@@ -274,15 +275,19 @@ public:
 	 */
 	void SetCompressionOptions(ECompressedBufferCompressor Compressor, ECompressedBufferCompressionLevel CompressionLevel);
 
+	UE_DEPRECATED(5.1, "Call GetBulkDataVersions instead.")
+	FCustomVersionContainer GetCustomVersions(FArchive& InlineArchive);
+
 	/**
-	* Get the CustomVersions used in the file containing the payload. Currently this is assumed
+	* Get the versions used in the file containing the payload. Currently this is assumed
 	* to always be the versions in the InlineArchive
-	* 
+	*
 	* @param InlineArchive The archive that was used to load this object
-	* 
+	*
 	* @return The CustomVersions that apply to the interpretation of the payload.
 	*/
-	FCustomVersionContainer GetCustomVersions(FArchive& InlineArchive);
+	void GetBulkDataVersions(FArchive& InlineArchive, FPackageFileVersion& OutUEVersion, int32& OutLicenseeUEVersion,
+		FCustomVersionContainer& OutCustomVersions) const;
 
 	/**
 	 * Set this BulkData into Torn-Off mode. It will no longer register with the BulkDataRegistry, even if
diff --git a/Engine/Source/Runtime/MeshDescription/Private/MeshDescription.cpp b/Engine/Source/Runtime/MeshDescription/Private/MeshDescription.cpp
--- a/Engine/Source/Runtime/MeshDescription/Private/MeshDescription.cpp
+++ b/Engine/Source/Runtime/MeshDescription/Private/MeshDescription.cpp
@@ -1999,7 +1999,9 @@ void FMeshDescriptionBulkData::Serialize( FArchive& Ar, UObject* Owner )
 	if ( Ar.IsTransacting() )
 	{
 		// If transacting, keep these members alive the other side of an undo, otherwise their values will get lost
-		CustomVersions.Serialize( Ar );
+		Ar << UEVersion;
+		Ar << LicenseeUEVersion;
+		CustomVersionsMrp.Serialize( Ar );
 		Ar << bBulkDataUpdated;
 	}
 	else
@@ -2054,9 +2056,10 @@ void FMeshDescriptionBulkData::Serialize( FArchive& Ar, UObject* Owner )
 
 	if (!Ar.IsTransacting() && Ar.IsLoading())
 	{
-		// If loading, take a copy of the package custom version container, so it can be applied when unpacking
+		// If loading, take a copy of the package's version information, so it can be applied when unpacking
 		// MeshDescription from the bulk data.
-		CustomVersions = BulkData.GetCustomVersions(Ar);
+		BulkData.GetBulkDataVersions(Ar, UEVersion, LicenseeUEVersion, CustomVersionsMrp);
+
 	}
 
 	// Needs to be after the guid is serialized
@@ -2083,8 +2086,10 @@ void FMeshDescriptionBulkData::SaveMeshDescription( FMeshDescription& MeshDescri
 		UE::Serialization::FEditorBulkDataWriter Ar(BulkData, bIsPersistent);
 		Ar << MeshDescription;
 
-		// Preserve CustomVersions at save time so we can reuse the same ones when reloading direct from memory
-		CustomVersions = Ar.GetCustomVersions();
+		// Preserve versions at save time so we can reuse the same ones when reloading direct from memory
+		UEVersion = Ar.UEVer();
+		LicenseeUEVersion = Ar.LicenseeUEVer();
+		CustomVersionsMrp = Ar.GetCustomVersions();
 	}
 
 	if (bGuidIsHash)
@@ -2118,9 +2123,11 @@ void FMeshDescriptionBulkData::LoadMeshDescription( FMeshDescription& MeshDescri
 		const bool bIsPersistent = true;
 		UE::Serialization::FEditorBulkDataReader Ar(BulkData, bIsPersistent);
 
-		// Propagate the custom version information from the package to the bulk data, so that the MeshDescription
+		// Propagate the version information from the package to the bulk data, so that the MeshDescription
 		// is serialized with the same versioning.
-		Ar.SetCustomVersions(CustomVersions);
+		Ar.SetUEVer(UEVersion);
+		Ar.SetLicenseeUEVer(LicenseeUEVersion);
+		Ar.SetCustomVersions(CustomVersionsMrp);
 		Ar << MeshDescription;
 
 		BulkData.UnloadData();
@@ -2150,7 +2157,9 @@ void FMeshDescriptionBulkData::UpdateMeshDescriptionFormat()
 	BulkData.UnloadData();
 	// Propagate the custom version information from the package to the bulk data, so that the MeshDescription
 	// is serialized with the same versioning.
-	Reader.SetCustomVersions(CustomVersions);
+	Reader.SetUEVer(UEVersion);
+	Reader.SetLicenseeUEVer(LicenseeUEVersion);
+	Reader.SetCustomVersions(CustomVersionsMrp);
 	FMeshDescription MeshDescription;
 	MeshDescription.Empty();
 	Reader << MeshDescription;
@@ -2167,15 +2176,17 @@ void FMeshDescriptionBulkData::UpdateMeshDescriptionFormat()
 	{
 		// The MeshDescription format has changed, so save it to this's BulkData
 		BulkData.UpdatePayload(FSharedBuffer::TakeOwnership(NewBytes.ReleaseOwnership(), NumBytes, FMemory::Free));
-		// Preserve CustomVersions at save time so we can reuse the same ones when reloading direct from memory
-		CustomVersions = NewBytes.GetCustomVersions();
+		// Preserve versions at save time so we can reuse the same ones when reloading direct from memory
+		UEVersion = NewBytes.UEVer();
+		LicenseeUEVersion = NewBytes.LicenseeUEVer();
+		CustomVersionsMrp = NewBytes.GetCustomVersions();
 		if (bGuidIsHash)
 		{
 			UseHashAsGuid();
 		}
 		else
 		{
-			// Maintain the original guid; this is a format change only. GetIdString will change because it adds the CustomVersions to the key.
+			// Maintain the original guid; this is a format change only. GetIdString will change because it adds the version information to the key.
 		}
 	}
 
@@ -2191,10 +2202,14 @@ void FMeshDescriptionBulkData::Empty()
 
 FString FMeshDescriptionBulkData::GetIdString() const
 {
-	// Create the IDString by combining this->Guid with the CustomVersions used for serialization,
-	// hashed down into a Guid to keep the string length the same as the original Guid.
+	// Create the IDString by combining this->Guid with the version information that will be used
+	// used for save serialization, hashed down into a Guid to keep the string length the same as the original Guid.
+	// Use the current binary's value for all versions, because this will be the versioning used during
+	// save serialization.
 	UE::DerivedData::FBuildVersionBuilder Builder;
 	Builder << const_cast<FGuid&>(Guid);
+	Builder << GPackageFileUEVersion;
+	Builder << GPackageFileLicenseeUEVersion;
 	for (const FGuid& CustomVersionGuid : GetMeshDescriptionCustomVersions())
 	{
 		Builder << const_cast<FGuid&>(CustomVersionGuid);
diff --git a/Engine/Source/Runtime/MeshDescription/Public/MeshDescription.h b/Engine/Source/Runtime/MeshDescription/Public/MeshDescription.h
--- a/Engine/Source/Runtime/MeshDescription/Public/MeshDescription.h
+++ b/Engine/Source/Runtime/MeshDescription/Public/MeshDescription.h
@@ -1367,7 +1367,10 @@ private:
 	FGuid GetHash() const;
 
 	/** Take a copy of the bulk data versioning so it can be propagated to the bulk data reader when deserializing MeshDescription */
-	FCustomVersionContainer CustomVersions;
+	FCustomVersionContainer CustomVersionsMrp;
+	FPackageFileVersion UEVersion;
+	int32 LicenseeUEVersion;
+
 
 	/** Whether the bulk data has been written via SaveMeshDescription */
 	bool bBulkDataUpdated;
