Removed TIsSame from misc. engine files (part of a larger change to deprecate TIsSame).

diff --git a/Engine/Source/Runtime/Engine/Classes/Animation/AnimNodeBase.h b/Engine/Source/Runtime/Engine/Classes/Animation/AnimNodeBase.h
--- a/Engine/Source/Runtime/Engine/Classes/Animation/AnimNodeBase.h
+++ b/Engine/Source/Runtime/Engine/Classes/Animation/AnimNodeBase.h
@@ -1028,7 +1028,7 @@ private:
 };
 
 #if WITH_EDITORONLY_DATA
-#define VERIFY_ANIM_NODE_MEMBER_TYPE(Type, Identifier) static_assert(TIsSame<decltype(Identifier), Type>::Value, "Incorrect return type used");
+#define VERIFY_ANIM_NODE_MEMBER_TYPE(Type, Identifier) static_assert(std::is_same_v<decltype(Identifier), Type>, "Incorrect return type used");
 #else
 #define VERIFY_ANIM_NODE_MEMBER_TYPE(Type, Identifier)
 #endif
diff --git a/Engine/Source/Runtime/Engine/Private/Collision/CollisionConversions.cpp b/Engine/Source/Runtime/Engine/Private/Collision/CollisionConversions.cpp
--- a/Engine/Source/Runtime/Engine/Private/Collision/CollisionConversions.cpp
+++ b/Engine/Source/Runtime/Engine/Private/Collision/CollisionConversions.cpp
@@ -416,7 +416,7 @@ EConvertQueryResult ConvertTraceResults(bool& OutHasValidBlockingHit, const UWor
 		HitType& Hit = Hits[i];
 		if(GetDistance(Hit) <= MaxDistance)
 		{
-			if (TIsSame<HitType, FHitSweep>::Value)
+			if constexpr (std::is_same_v<HitType, FHitSweep>)
 			{
 				if (!HadInitialOverlap(Hit))
 				{
@@ -449,7 +449,7 @@ template <typename Hit>
 EConvertQueryResult ConvertTraceResults(bool& OutHasValidBlockingHit, const UWorld* World, int32 NumHits, Hit* Hits, float CheckLength, const FCollisionFilterData& QueryFilter, FHitResult& OutHit, const FVector& StartLoc, const FVector& EndLoc, const FPhysicsGeometry& Geom, const FTransform& QueryTM, float MaxDistance, bool bReturnFaceIndex, bool bReturnPhysMat)
 {
 	const FVector Dir = (EndLoc - StartLoc).GetSafeNormal();
-	if (TIsSame<Hit, FHitSweep>::Value)
+	if constexpr (std::is_same_v<Hit, FHitSweep>)
 	{
 		if (!HadInitialOverlap(Hits[0]))
 		{
diff --git a/Engine/Source/Runtime/Engine/Public/Animation/AnimNodeMessages.h b/Engine/Source/Runtime/Engine/Public/Animation/AnimNodeMessages.h
--- a/Engine/Source/Runtime/Engine/Public/Animation/AnimNodeMessages.h
+++ b/Engine/Source/Runtime/Engine/Public/Animation/AnimNodeMessages.h
@@ -155,7 +155,7 @@ struct TScopedGraphMessage : FScopedGraphMessage
 		: FScopedGraphMessage(InContext)
 	{
 		static_assert(TIsDerivedFrom<TGraphMessageType, IGraphMessage>::IsDerived, "Argument TGraphMessageType must derive from IGraphMessage");
-		static_assert(!TIsSame<TGraphMessageType, IGraphMessage>::Value, "Argument TGraphMessageType must not be IGraphMessage");
+		static_assert(!std::is_same_v<TGraphMessageType, IGraphMessage>, "Argument TGraphMessageType must not be IGraphMessage");
 
 		TSharedRef<TGraphMessageType> Message = MakeShared<TGraphMessageType>(Forward<TArgs>(Args)...);
 		PushMessage(InContext, Message, Message->GetTypeName());
@@ -177,7 +177,7 @@ struct TOptionalScopedGraphMessage : FScopedGraphMessage
 		, bCondition(bInCondition)
 	{
 		static_assert(TIsDerivedFrom<TGraphMessageType, IGraphMessage>::IsDerived, "Argument TGraphMessageType must derive from IGraphMessage");
-		static_assert(!TIsSame<TGraphMessageType, IGraphMessage>::Value, "Argument TGraphMessageType must not be IGraphMessage");
+		static_assert(!std::is_same_v<TGraphMessageType, IGraphMessage>, "Argument TGraphMessageType must not be IGraphMessage");
 
 		if(bInCondition)
 		{
@@ -235,7 +235,7 @@ public:
 	void ForEachMessage(TFunctionRef<EEnumerate(TGraphMessageType&)> InFunction) const
 	{
 		static_assert(TIsDerivedFrom<TGraphMessageType, IGraphMessage>::IsDerived, "Argument TGraphMessageType must derive from IGraphMessage");
-		static_assert(!TIsSame<TGraphMessageType, IGraphMessage>::Value, "Argument TGraphMessageType must not be IGraphMessage");
+		static_assert(!std::is_same_v<TGraphMessageType, IGraphMessage>, "Argument TGraphMessageType must not be IGraphMessage");
 
 		ForEachMessage(TGraphMessageType::GetStaticTypeName(), [&InFunction](IGraphMessage& InMessage)
 		{
@@ -249,7 +249,7 @@ public:
 	void TopMessage(TFunctionRef<void(TGraphMessageType&)> InFunction) const
 	{
 		static_assert(TIsDerivedFrom<TGraphMessageType, IGraphMessage>::IsDerived, "Argument TGraphMessageType must derive from IGraphMessage");
-		static_assert(!TIsSame<TGraphMessageType, IGraphMessage>::Value, "Argument TGraphMessageType must not be IGraphMessage");
+		static_assert(!std::is_same_v<TGraphMessageType, IGraphMessage>, "Argument TGraphMessageType must not be IGraphMessage");
 
 		TopMessage(TGraphMessageType::GetStaticTypeName(), [&InFunction](IGraphMessage& InMessage)
 		{
@@ -262,7 +262,7 @@ public:
 	bool HasMessage() const
 	{
 		static_assert(TIsDerivedFrom<TGraphMessageType, IGraphMessage>::IsDerived, "Argument TGraphMessageType must derive from IGraphMessage");
-		static_assert(!TIsSame<TGraphMessageType, IGraphMessage>::Value, "Argument TGraphMessageType must not be IGraphMessage");
+		static_assert(!std::is_same_v<TGraphMessageType, IGraphMessage>, "Argument TGraphMessageType must not be IGraphMessage");
 
 		return HasMessage(TGraphMessageType::GetStaticTypeName());
 	}
diff --git a/Engine/Source/Runtime/ImageWrapper/Private/Formats/TiffImageWrapper.cpp b/Engine/Source/Runtime/ImageWrapper/Private/Formats/TiffImageWrapper.cpp
--- a/Engine/Source/Runtime/ImageWrapper/Private/Formats/TiffImageWrapper.cpp
+++ b/Engine/Source/Runtime/ImageWrapper/Private/Formats/TiffImageWrapper.cpp
@@ -30,7 +30,7 @@ namespace UE::ImageWrapper::Private
 		template<class DataTypeDest, class DataTypeSrc>
 		DataTypeDest ConvertToWriteFormat(DataTypeSrc ReadValue)
 		{
-			if constexpr ( TIsSame<DataTypeDest, DataTypeSrc>::Value )
+			if constexpr (std::is_same_v<DataTypeDest, DataTypeSrc>)
 			{
 				return ReadValue;
 			}
@@ -42,7 +42,7 @@ namespace UE::ImageWrapper::Private
 
 				// get source value in [0,1]
 				double SrcValue;
-				if constexpr (TIsSame<DataTypeSrc, FFloat16>::Value || TIsFloatingPoint<DataTypeSrc>::Value)
+				if constexpr (std::is_same_v<DataTypeSrc, FFloat16> || TIsFloatingPoint<DataTypeSrc>::Value)
 				{
 					SrcValue = FMath::Clamp((double)ReadValue, 0.0, 1.0);
 				}
@@ -52,7 +52,7 @@ namespace UE::ImageWrapper::Private
 				}
 				
 				// quantize to dest :
-				if constexpr (TIsSame<DataTypeDest, FFloat16>::Value || TIsFloatingPoint<DataTypeDest>::Value)
+				if constexpr (std::is_same_v<DataTypeDest, FFloat16> || TIsFloatingPoint<DataTypeDest>::Value)
 				{
 					return DataTypeDest(SrcValue);
 				}
@@ -767,7 +767,7 @@ namespace UE::ImageWrapper::Private
 
 		const bool bIsTiled = TIFFIsTiled(Tiff) != 0;
 
-		if constexpr (TIsSame<DataTypeDest, uint8>::Value)
+		if constexpr (std::is_same_v<DataTypeDest, uint8>)
 		{
 			if (BitsPerSample == 1)
 			{
@@ -796,7 +796,7 @@ namespace UE::ImageWrapper::Private
 			{
 				if (Photometric == PHOTOMETRIC_PALETTE)
 				{
-					if constexpr (TIsSame<DataTypeDest, uint16>::Value)
+					if constexpr (std::is_same_v<DataTypeDest, uint16>)
 					{
 						switch (BitsPerSample)
 						{
@@ -869,7 +869,7 @@ namespace UE::ImageWrapper::Private
 		
 		if (Photometric == PHOTOMETRIC_MINISWHITE)
 		{
-			if constexpr ( TIsSame<DataTypeDest, FFloat16>::Value || TIsFloatingPoint<DataTypeDest>::Value )
+			if constexpr (std::is_same_v<DataTypeDest, FFloat16> || TIsFloatingPoint<DataTypeDest>::Value)
 			{
 				// miniswhite with floating point tiff is an odd thing to do
 				// it's unclear if you should invert against 1.f or something else?
@@ -1096,7 +1096,7 @@ namespace UE::ImageWrapper::Private
 			ParallelFor(Width * Height, [NumOfChannelDest, &WriteArray](int32 Index)
 				{
 					const int64 WriteIndex = int64(NumOfChannelDest) * Index + 3;
-					if constexpr ( TIsSame<DataTypeDest, FFloat16>::Value || TIsFloatingPoint<DataTypeDest>::Value ) 
+					if constexpr (std::is_same_v<DataTypeDest, FFloat16> || TIsFloatingPoint<DataTypeDest>::Value)
 					{
 						WriteArray[WriteIndex] = DataTypeDest(1.f);
 					}
diff --git a/Engine/Source/Runtime/JsonUtilities/Public/JsonDomBuilder.h b/Engine/Source/Runtime/JsonUtilities/Public/JsonDomBuilder.h
--- a/Engine/Source/Runtime/JsonUtilities/Public/JsonDomBuilder.h
+++ b/Engine/Source/Runtime/JsonUtilities/Public/JsonDomBuilder.h
@@ -71,11 +71,11 @@ public:
 		FObject& Set(const FString& Key, const FString& Str)           { Object->SetField(Key, MakeShared<FJsonValueString>(Str));            return *this; }
 
 		template <class FNumber>
-		typename TEnableIf<!TIsSame<FNumber, bool>::Value && (TIsIntegral<FNumber>::Value || TIsFloatingPoint<FNumber>::Value), FObject&>::Type
+		typename TEnableIf<!std::is_same_v<FNumber, bool> && (TIsIntegral<FNumber>::Value || TIsFloatingPoint<FNumber>::Value), FObject&>::Type
 			Set(const FString& Key, FNumber Number)                    { Object->SetField(Key, MakeShared<FJsonValueNumber>(Number));         return *this; }
 
 		template <class FBool>
-		typename TEnableIf<TIsSame<FBool, bool>::Value, FObject&>::Type
+		typename TEnableIf<std::is_same_v<FBool, bool>, FObject&>::Type
 			Set(const FString& Key, FBool Boolean)                     { Object->SetField(Key, MakeShared<FJsonValueBoolean>(Boolean));       return *this; }
 
 		FObject& Set(const FString& Key, TYPE_OF_NULLPTR)              { Object->SetField(Key, MakeShared<FJsonValueNull>());                 return *this; }
diff --git a/Engine/Source/Runtime/LiveLinkInterface/Public/LiveLinkTypes.h b/Engine/Source/Runtime/LiveLinkInterface/Public/LiveLinkTypes.h
--- a/Engine/Source/Runtime/LiveLinkInterface/Public/LiveLinkTypes.h
+++ b/Engine/Source/Runtime/LiveLinkInterface/Public/LiveLinkTypes.h
@@ -472,7 +472,7 @@ protected:
 	{
 		static Type* Cast(const UScriptStruct* ScriptStruct, BaseType* BaseData)
 		{
-			if (TIsSame<Type, BaseType>::Value)
+			if constexpr (std::is_same_v<Type, BaseType>)
 			{
 				return StaticCast<Type*>(BaseData);
 			}
@@ -489,7 +489,7 @@ protected:
 		}
 		static const Type* ConstCast(const UScriptStruct* ScriptStruct, const BaseType* BaseData)
 		{
-			if (TIsSame<Type, BaseType>::Value)
+			if constexpr (std::is_same_v<Type, BaseType>)
 			{
 				return StaticCast<const Type*>(BaseData);
 			}
diff --git a/Engine/Source/Runtime/MeshDescription/Public/AttributeArrayContainer.h b/Engine/Source/Runtime/MeshDescription/Public/AttributeArrayContainer.h
--- a/Engine/Source/Runtime/MeshDescription/Public/AttributeArrayContainer.h
+++ b/Engine/Source/Runtime/MeshDescription/Public/AttributeArrayContainer.h
@@ -403,7 +403,7 @@ public:
 	/**
 	 * Construct a TArrayAttribute<const T> from a TArrayAttribute<T>. 
 	 */
-	template <typename T = AttributeType, typename TEnableIf<TIsSame<T, const T>::Value, int>::Type = 0>
+	template <typename T = AttributeType, typename TEnableIf<std::is_same_v<T, const T>, int>::Type = 0>
 	TArrayAttribute(TArrayAttribute<typename TRemoveCV<T>::Type> InValue)
 		: Array(InValue.Array),
 		  Index(InValue.Index)
diff --git a/Engine/Source/Runtime/MeshDescription/Public/MeshAttributeArray.h b/Engine/Source/Runtime/MeshDescription/Public/MeshAttributeArray.h
--- a/Engine/Source/Runtime/MeshDescription/Public/MeshAttributeArray.h
+++ b/Engine/Source/Runtime/MeshDescription/Public/MeshAttributeArray.h
@@ -816,19 +816,19 @@ public:
 	{}
 
 	/** Implicitly construct a TMeshAttributesRef-to-const from a regular one */
-	template <typename T = AttributeType, typename TEnableIf<TIsSame<T, const T>::Value, int>::Type = 0>
+	template <typename T = AttributeType, typename TEnableIf<std::is_same_v<T, const T>, int>::Type = 0>
 	TMeshAttributesRef(TMeshAttributesRef<ElementIDType, typename TRemoveCV<T>::Type> InRef)
 		: ArrayPtr(InRef.ArrayPtr)
 	{}
 
 	/** Implicitly construct a TMeshAttributesRef from a TMeshAttributesArray **/
-	template <typename T = ElementIDType, typename U = AttributeType, typename TEnableIf<!TIsSame<T, int32>::Value, int>::Type = 0>
+	template <typename T = ElementIDType, typename U = AttributeType, typename TEnableIf<!std::is_same_v<T, int32>, int>::Type = 0>
 	TMeshAttributesRef(TMeshAttributesRef<int32, U> InRef)
 		: ArrayPtr(InRef.ArrayPtr)
 	{}
 
 	/** Implicitly construct a TMeshAttributesRef-to-const from a TMeshAttributesArray */
-	template <typename T = ElementIDType, typename U = AttributeType, typename TEnableIf<TIsSame<U, const U>::Value, int>::Type = 0, typename TEnableIf<!TIsSame<T, int32>::Value, int>::Type = 0>
+	template <typename T = ElementIDType, typename U = AttributeType, typename TEnableIf<std::is_same_v<U, const U>, int>::Type = 0, typename TEnableIf<!std::is_same_v<T, int32>, int>::Type = 0>
 	TMeshAttributesRef(TMeshAttributesRef<int32, typename TRemoveCV<U>::Type> InRef)
 		: ArrayPtr(InRef.ArrayPtr)
 	{}
@@ -1013,21 +1013,21 @@ public:
 	{}
 
 	/** Implicitly construct a TMeshAttributesRef-to-const from a regular one */
-	template <typename T = AttributeType, typename TEnableIf<TIsSame<T, const T>::Value, int>::Type = 0>
+	template <typename T = AttributeType, typename TEnableIf<std::is_same_v<T, const T>, int>::Type = 0>
 	TMeshAttributesRef(TMeshAttributesRef<ElementIDType, TArrayView<typename TRemoveCV<T>::Type>> InRef)
 		: ArrayPtr(InRef.ArrayPtr),
 		  Extent(InRef.Extent)
 	{}
 
 	/** Implicitly construct a TMeshAttributesRef from a TMeshAttributesArray **/
-	template <typename T = ElementIDType, typename U = AttributeType, typename TEnableIf<!TIsSame<T, int32>::Value, int>::Type = 0>
+	template <typename T = ElementIDType, typename U = AttributeType, typename TEnableIf<!std::is_same_v<T, int32>, int>::Type = 0>
 	TMeshAttributesRef(TMeshAttributesRef<int32, TArrayView<U>> InRef)
 		: ArrayPtr(InRef.ArrayPtr),
 		  Extent(InRef.Extent)
 	{}
 
 	/** Implicitly construct a TMeshAttributesRef-to-const from a TMeshAttributesArray */
-	template <typename T = ElementIDType, typename U = AttributeType, typename TEnableIf<TIsSame<U, const U>::Value, int>::Type = 0, typename TEnableIf<!TIsSame<T, int32>::Value, int>::Type = 0>
+	template <typename T = ElementIDType, typename U = AttributeType, typename TEnableIf<std::is_same_v<U, const U>, int>::Type = 0, typename TEnableIf<!std::is_same_v<T, int32>, int>::Type = 0>
 	TMeshAttributesRef(TMeshAttributesRef<int32, TArrayView<typename TRemoveCV<U>::Type>> InRef)
 		: ArrayPtr(InRef.ArrayPtr),
 		  Extent(InRef.Extent)
@@ -1278,19 +1278,19 @@ public:
 	{}
 
 	/** Implicitly construct a TMeshAttributesRef-to-const from a regular one */
-	template <typename T = AttributeType, typename TEnableIf<TIsSame<T, const T>::Value, int>::Type = 0>
+	template <typename T = AttributeType, typename TEnableIf<std::is_same_v<T, const T>, int>::Type = 0>
 	TMeshAttributesRef(TMeshAttributesRef<ElementIDType, TArrayAttribute<typename TRemoveCV<T>::Type>> InRef)
 		: ArrayPtr(InRef.ArrayPtr)
 	{}
 
 	/** Implicitly construct a TMeshAttributesRef from a TMeshAttributesArray **/
-	template <typename T = ElementIDType, typename U = AttributeType, typename TEnableIf<!TIsSame<T, int32>::Value, int>::Type = 0>
+	template <typename T = ElementIDType, typename U = AttributeType, typename TEnableIf<!std::is_same_v<T, int32>, int>::Type = 0>
 	TMeshAttributesRef(TMeshAttributesRef<int32, TArrayAttribute<U>> InRef)
 		: ArrayPtr(InRef.ArrayPtr)
 	{}
 
 	/** Implicitly construct a TMeshAttributesRef-to-const from a TMeshAttributesArray */
-	template <typename T = ElementIDType, typename U = AttributeType, typename TEnableIf<TIsSame<U, const U>::Value, int>::Type = 0, typename TEnableIf<!TIsSame<T, int32>::Value, int>::Type = 0>
+	template <typename T = ElementIDType, typename U = AttributeType, typename TEnableIf<std::is_same_v<U, const U>, int>::Type = 0, typename TEnableIf<!std::is_same_v<T, int32>, int>::Type = 0>
 	TMeshAttributesRef(TMeshAttributesRef<int32, TArrayAttribute<typename TRemoveCV<U>::Type>> InRef)
 		: ArrayPtr(InRef.ArrayPtr)
 	{}
@@ -1587,7 +1587,7 @@ public:
 	 * Obtain a reference to this with GetAttributesRef<TArrayAttribute<float>>(...)
 	 */
 	template <typename T,
-			  typename TEnableIf<TIsSame<typename TMeshAttributesRegisterType<T>::RealAttributeType, int>::Value, int>::Type = 0>
+			  typename TEnableIf<std::is_same_v<typename TMeshAttributesRegisterType<T>::RealAttributeType, int>, int>::Type = 0>
 	TMeshAttributesArray<typename TMeshAttributesRegisterType<T>::RefType> RegisterIndexAttribute(
 		const FName AttributeName,
 		const int32 NumberOfChannels = 1,
diff --git a/Engine/Source/Runtime/NavigationSystem/Private/NavMesh/PImplRecastNavMesh.cpp b/Engine/Source/Runtime/NavigationSystem/Private/NavMesh/PImplRecastNavMesh.cpp
--- a/Engine/Source/Runtime/NavigationSystem/Private/NavMesh/PImplRecastNavMesh.cpp
+++ b/Engine/Source/Runtime/NavigationSystem/Private/NavMesh/PImplRecastNavMesh.cpp
@@ -31,9 +31,9 @@ static_assert(sizeof(NavNodeRef) == sizeof(dtPolyRef), "NavNodeRef and dtPolyRef
 static_assert(RECAST_MAX_AREAS <= DT_MAX_AREAS, "Number of allowed areas cannot exceed DT_MAX_AREAS.");
 static_assert(RECAST_STRAIGHTPATH_OFFMESH_CONNECTION == DT_STRAIGHTPATH_OFFMESH_CONNECTION, "Path flags values differ.");
 static_assert(RECAST_UNWALKABLE_POLY_COST == DT_UNWALKABLE_POLY_COST, "Unwalkable poly cost differ.");
-static_assert(TIsSame<FVector::FReal, dtReal>::Value, "FReal and dtReal must be the same type!");
-static_assert(TIsSame<FVector::FReal, rcReal>::Value, "FReal and rcReal must be the same type!");
-static_assert(TIsSame<FVector::FReal, duReal>::Value, "FReal and duReal must be the same type!");
+static_assert(std::is_same_v<FVector::FReal, dtReal>, "FReal and dtReal must be the same type!");
+static_assert(std::is_same_v<FVector::FReal, rcReal>, "FReal and rcReal must be the same type!");
+static_assert(std::is_same_v<FVector::FReal, duReal>, "FReal and duReal must be the same type!");
 
 /// Helper for accessing navigation query from different threads
 #define INITIALIZE_NAVQUERY_SIMPLE(NavQueryVariable, NumNodes)	\
diff --git a/Engine/Source/Runtime/Net/Core/Classes/Net/Serialization/FastArraySerializer.h b/Engine/Source/Runtime/Net/Core/Classes/Net/Serialization/FastArraySerializer.h
--- a/Engine/Source/Runtime/Net/Core/Classes/Net/Serialization/FastArraySerializer.h
+++ b/Engine/Source/Runtime/Net/Core/Classes/Net/Serialization/FastArraySerializer.h
@@ -709,7 +709,7 @@ private:
 		inline typename TEnableIf<!TModels<CPostReplicatedReceiveFuncable, FastArrayType, const FFastArraySerializer::FPostReplicatedReceiveParameters>::Value, void>::Type CallPostReplicatedReceiveOrNot(int32) {}
 
 		// Validate that deduced FastArrayItemType is valid and that it is the same as the specified one		
-		static_assert(TIsSame<typename TFastArrayTypeHelper<SerializerType>::FastArrayItemType, Type>::Value, "Auto deduced FastArrayItemType is invalid or differs from the specified type. Make sure that the FastArraySerializer has a single replicated array property.");
+		static_assert(std::is_same_v<typename TFastArrayTypeHelper<SerializerType>::FastArrayItemType, Type>, "Auto deduced FastArrayItemType is invalid or differs from the specified type. Make sure that the FastArraySerializer has a single replicated array property.");
 	};
 
 	/**
diff --git a/Engine/Source/Runtime/RenderCore/Public/RenderGraphPass.h b/Engine/Source/Runtime/RenderCore/Public/RenderGraphPass.h
--- a/Engine/Source/Runtime/RenderCore/Public/RenderGraphPass.h
+++ b/Engine/Source/Runtime/RenderCore/Public/RenderGraphPass.h
@@ -606,18 +606,18 @@ public:
 		, DebugParameterStruct(InParameterStruct)
 #endif
 	{		
-		bParallelExecuteAllowed = !TIsSame<TRHICommandList, FRHICommandListImmediate>::Value && !EnumHasAnyFlags(InPassFlags, ERDGPassFlags::NeverParallel);
+		bParallelExecuteAllowed = !std::is_same_v<TRHICommandList, FRHICommandListImmediate> && !EnumHasAnyFlags(InPassFlags, ERDGPassFlags::NeverParallel);
 	}
 
 private:
 	template<class T>
-	typename TEnableIf<!TIsSame<T, FRDGPass>::Value, void>::Type ExecuteLambdaFunc(FRHIComputeCommandList& RHICmdList)
+	typename TEnableIf<!std::is_same_v<T, FRDGPass>, void>::Type ExecuteLambdaFunc(FRHIComputeCommandList& RHICmdList)
 	{
 		ExecuteLambda(static_cast<TRHICommandList&>(RHICmdList));
 	}
 
 	template<class T>
-	typename TEnableIf<TIsSame<T, FRDGPass>::Value, void>::Type ExecuteLambdaFunc(FRHIComputeCommandList& RHICmdList)
+	typename TEnableIf<std::is_same_v<T, FRDGPass>, void>::Type ExecuteLambdaFunc(FRHIComputeCommandList& RHICmdList)
 	{
 		ExecuteLambda(this, static_cast<TRHICommandList&>(RHICmdList));
 	}
diff --git a/Engine/Source/Runtime/RenderCore/Public/RenderGraphUtils.h b/Engine/Source/Runtime/RenderCore/Public/RenderGraphUtils.h
--- a/Engine/Source/Runtime/RenderCore/Public/RenderGraphUtils.h
+++ b/Engine/Source/Runtime/RenderCore/Public/RenderGraphUtils.h
@@ -224,8 +224,8 @@ void ClearUnusedGraphResources(
 	TPassParameterStruct* InoutParameters,
 	std::initializer_list<FRDGResourceRef> ExcludeList = {})
 {
-	static_assert(TIsSame<typename TShaderClassA::FParameters, TPassParameterStruct>::Value, "First shader FParameter type must match pass parameters.");
-	static_assert(TIsSame<typename TShaderClassB::FParameters, TPassParameterStruct>::Value, "Second shader FParameter type must match pass parameters.");
+	static_assert(std::is_same_v<typename TShaderClassA::FParameters, TPassParameterStruct>, "First shader FParameter type must match pass parameters.");
+	static_assert(std::is_same_v<typename TShaderClassB::FParameters, TPassParameterStruct>, "Second shader FParameter type must match pass parameters.");
 	const FShaderParametersMetadata* ParametersMetadata = TPassParameterStruct::FTypeInfo::GetStructMetadata();
 
 	// Verify the shader have all the parameters it needs. This is done before the
diff --git a/Engine/Source/Runtime/RenderCore/Public/Shader.h b/Engine/Source/Runtime/RenderCore/Public/Shader.h
--- a/Engine/Source/Runtime/RenderCore/Public/Shader.h
+++ b/Engine/Source/Runtime/RenderCore/Public/Shader.h
@@ -2349,7 +2349,7 @@ public:
 		: TShaderRef<ShaderType>(ShaderIndex->template GetShader<ShaderType>(/* PermutationId = */ 0)) // gcc3 needs the template quantifier so it knows the < is not a less-than
 	{
 		static_assert(
-			TIsSame<typename ShaderType::FPermutationDomain, FShaderPermutationNone>::Value,
+			std::is_same_v<typename ShaderType::FPermutationDomain, FShaderPermutationNone>,
 			"Missing permutation vector argument for shader that have a permutation domain.");
 	}
 
diff --git a/Engine/Source/Runtime/RenderCore/Public/ShaderPermutation.h b/Engine/Source/Runtime/RenderCore/Public/ShaderPermutation.h
--- a/Engine/Source/Runtime/RenderCore/Public/ShaderPermutation.h
+++ b/Engine/Source/Runtime/RenderCore/Public/ShaderPermutation.h
@@ -368,7 +368,7 @@ struct TShaderPermutationDomain<TDimension, Ts...>
 	template<class DimensionToSet>
 	void Set(typename DimensionToSet::Type Value)
 	{
-		return TShaderPermutationDomainSpetialization<TIsSame<TDimension, DimensionToSet>::Value>::template SetDimension<Type, DimensionToSet>(*this, Value);
+		return TShaderPermutationDomainSpetialization<std::is_same_v<TDimension, DimensionToSet>>::template SetDimension<Type, DimensionToSet>(*this, Value);
 	}
 
 
@@ -376,7 +376,7 @@ struct TShaderPermutationDomain<TDimension, Ts...>
 	template<class DimensionToGet>
 	const typename DimensionToGet::Type& Get() const
 	{
-		return TShaderPermutationDomainSpetialization<TIsSame<TDimension, DimensionToGet>::Value>::template GetDimension<Type, DimensionToGet>(*this);
+		return TShaderPermutationDomainSpetialization<std::is_same_v<TDimension, DimensionToGet>>::template GetDimension<Type, DimensionToGet>(*this);
 	}
 
 
diff --git a/Engine/Source/Runtime/Renderer/Private/DeferredShadingRenderer.cpp b/Engine/Source/Runtime/Renderer/Private/DeferredShadingRenderer.cpp
--- a/Engine/Source/Runtime/Renderer/Private/DeferredShadingRenderer.cpp
+++ b/Engine/Source/Runtime/Renderer/Private/DeferredShadingRenderer.cpp
@@ -1226,7 +1226,7 @@ bool FDeferredShadingSceneRenderer::GatherRayTracingWorldInstancesForView(FRDGBu
 							checkf(Instance.InstanceTransformsView.Num() == 0, TEXT("InstanceTransformsView is expected to be empty if using InstanceTransforms"));
 							TArrayView<FMatrix> SceneOwnedTransforms = RayTracingScene.Allocate<FMatrix>(Instance.InstanceTransforms.Num());
 							FMemory::Memcpy(SceneOwnedTransforms.GetData(), Instance.InstanceTransforms.GetData(), Instance.InstanceTransforms.Num() * sizeof(RayTracingInstance.Transforms[0]));
-							static_assert(TIsSame<decltype(SceneOwnedTransforms[0]), decltype(Instance.InstanceTransforms[0])>::Value, "Unexpected transform type");
+							static_assert(std::is_same_v<decltype(SceneOwnedTransforms[0]), decltype(Instance.InstanceTransforms[0])>, "Unexpected transform type");
 
 							RayTracingInstance.NumTransforms = SceneOwnedTransforms.Num();
 							RayTracingInstance.Transforms = SceneOwnedTransforms;
diff --git a/Engine/Source/Runtime/SignalProcessing/Public/DSP/Dsp.h b/Engine/Source/Runtime/SignalProcessing/Public/DSP/Dsp.h
--- a/Engine/Source/Runtime/SignalProcessing/Public/DSP/Dsp.h
+++ b/Engine/Source/Runtime/SignalProcessing/Public/DSP/Dsp.h
@@ -1202,7 +1202,7 @@ namespace Audio
 		{
 			CheckValidityOfSampleType<OtherSampleType>();
 
-			if (TIsSame<SampleType, OtherSampleType>::Value)
+			if constexpr (std::is_same_v<SampleType, OtherSampleType>)
 			{
 				Sample = InSample;
 				return *this;
@@ -1391,7 +1391,7 @@ namespace Audio
 		{
 			CheckValidityOfSampleType<OtherSampleType>();
 
-			if (TIsSame<SampleType, OtherSampleType>::Value)
+			if constexpr (std::is_same_v<SampleType, OtherSampleType>)
 			{
 				Sample = InSample;
 				return *this;
diff --git a/Engine/Source/Runtime/SignalProcessing/Public/SampleBuffer.h b/Engine/Source/Runtime/SignalProcessing/Public/SampleBuffer.h
--- a/Engine/Source/Runtime/SignalProcessing/Public/SampleBuffer.h
+++ b/Engine/Source/Runtime/SignalProcessing/Public/SampleBuffer.h
@@ -76,11 +76,11 @@ namespace Audio
 			RawPCMData.Reset(NumSamples);
 			RawPCMData.AddUninitialized(NumSamples);
 
-			if constexpr(TIsSame<SampleType, float>::Value)
+			if constexpr(std::is_same_v<SampleType, float>)
 			{
 				FMemory::Memcpy(RawPCMData.GetData(), InBufferPtr, NumSamples * sizeof(float));
 			}
-			else if constexpr(TIsSame<SampleType, int16>::Value)
+			else if constexpr(std::is_same_v<SampleType, int16>)
 			{
 				// Convert from float to int:
 				for (int32 SampleIndex = 0; SampleIndex < NumSamples; SampleIndex++)
@@ -109,11 +109,11 @@ namespace Audio
 			RawPCMData.Reset(NumSamples);
 			RawPCMData.AddUninitialized(NumSamples);
 
-			if constexpr(TIsSame<SampleType, int16>::Value)
+			if constexpr(std::is_same_v<SampleType, int16>)
 			{
 				FMemory::Memcpy(RawPCMData.GetData(), InBufferPtr, NumSamples * sizeof(int16));
 			}
-			else if constexpr(TIsSame<SampleType, float>::Value)
+			else if constexpr(std::is_same_v<SampleType, float>)
 			{
 				// Convert from int to float:
 				for (int32 SampleIndex = 0; SampleIndex < NumSamples; SampleIndex++)
@@ -161,12 +161,12 @@ namespace Audio
 			RawPCMData.Reset(NumSamples);
 			RawPCMData.AddUninitialized(NumSamples);
 
-			if constexpr(TIsSame<SampleType, OtherSampleType>::Value)
+			if constexpr(std::is_same_v<SampleType, OtherSampleType>)
 			{
 				// If buffers are of the same type, copy over:
 				FMemory::Memcpy(RawPCMData.GetData(), Other.RawPCMData.GetData(), NumSamples * sizeof(SampleType));
 			}
-			else if constexpr(TIsSame<SampleType, int16>::Value && TIsSame<OtherSampleType, float>::Value)
+			else if constexpr(std::is_same_v<SampleType, int16> && std::is_same_v<OtherSampleType, float>)
 			{
 				// Convert from float to int:
 				for (int32 SampleIndex = 0; SampleIndex < NumSamples; SampleIndex++)
@@ -174,7 +174,7 @@ namespace Audio
 					RawPCMData[SampleIndex] = (int16)(Other.RawPCMData[SampleIndex] * 32767.0f);
 				}
 			}
-			else if constexpr(TIsSame<SampleType, float>::Value && TIsSame<OtherSampleType, int16>::Value)
+			else if constexpr(std::is_same_v<SampleType, float> && std::is_same_v<OtherSampleType, int16>)
 			{
 				// Convert from int to float:
 				for (int32 SampleIndex = 0; SampleIndex < NumSamples; SampleIndex++)
@@ -215,13 +215,13 @@ namespace Audio
 		{
 			int32 StartIndex = RawPCMData.AddUninitialized(InNumSamples);
 
-			if constexpr(TIsSame<SampleType, OtherSampleType>::Value)
+			if constexpr(std::is_same_v<SampleType, OtherSampleType>)
 			{
 				FMemory::Memcpy(&RawPCMData[StartIndex], InputBuffer, InNumSamples * sizeof(SampleType));
 			}
 			else
 			{
-				if constexpr(TIsSame<SampleType, int16>::Value && TIsSame<OtherSampleType, float>::Value)
+				if constexpr(std::is_same_v<SampleType, int16> && std::is_same_v<OtherSampleType, float>)
 				{
 					// Convert from float to int:
 					for (int32 SampleIndex = 0; SampleIndex < InNumSamples; SampleIndex++)
@@ -229,7 +229,7 @@ namespace Audio
 						RawPCMData[StartIndex + SampleIndex] = (int16)(InputBuffer[SampleIndex] * 32767.0f);
 					}
 				}
-				else if constexpr(TIsSame<SampleType, float>::Value && TIsSame<OtherSampleType, int16>::Value)
+				else if constexpr(std::is_same_v<SampleType, float> && std::is_same_v<OtherSampleType, int16>)
 				{
 					// Convert from int to float:
 					for (int32 SampleIndex = 0; SampleIndex < InNumSamples; SampleIndex++)
@@ -358,7 +358,7 @@ namespace Audio
 
 		void Clamp(float Ceiling = 1.0f)
 		{
-			if constexpr(TIsSame<SampleType, float>::Value)
+			if constexpr(std::is_same_v<SampleType, float>)
 			{
 				// Float case:
 				float ClampMin = Ceiling * -1.0f;
@@ -368,7 +368,7 @@ namespace Audio
 					RawPCMData[SampleIndex] = static_cast<SampleType>(FMath::Clamp<float>(RawPCMData[SampleIndex], ClampMin, Ceiling));
 				}
 			}
-			else if constexpr(TIsSame<SampleType, int16>::Value)
+			else if constexpr(std::is_same_v<SampleType, int16>)
 			{
 				// int16 case:
 				Ceiling = FMath::Clamp(Ceiling, 0.0f, 1.0f);
@@ -480,7 +480,7 @@ namespace Audio
 			{
 				float SampleA, SampleB;
 
-				if constexpr(TIsSame<SampleType, float>::Value)
+				if constexpr(std::is_same_v<SampleType, float>)
 				{
 					SampleA = RawPCMData[(WholeThisIndex * NumChannels) + i];
 					SampleB = RawPCMData[(WholeNextIndex * NumChannels) + i];
diff --git a/Engine/Source/Runtime/VulkanRHI/Private/VulkanCommands.cpp b/Engine/Source/Runtime/VulkanRHI/Private/VulkanCommands.cpp
--- a/Engine/Source/Runtime/VulkanRHI/Private/VulkanCommands.cpp
+++ b/Engine/Source/Runtime/VulkanRHI/Private/VulkanCommands.cpp
@@ -285,8 +285,8 @@ template <typename TState>
 inline void SetShaderUniformBufferResources(FVulkanCommandListContext* Context, TState* State, const FVulkanShader* Shader, const TArray<FVulkanShaderHeader::FGlobalInfo>& GlobalInfos, const TArray<TEnumAsByte<EVulkanBindingType::EType>>& DescriptorTypes, const FVulkanShaderHeader::FUniformBufferInfo& HeaderUBInfo, const FVulkanUniformBuffer* UniformBuffer, const TArray<FDescriptorSetRemappingInfo::FRemappingInfo>& GlobalRemappingInfo)
 {
 #if ENABLE_RHI_VALIDATION
-	static_assert(TIsSame<TState, FVulkanPendingGfxState>::Value || TIsSame<TState, FVulkanPendingComputeState>::Value, "TState must be FVulkanPendingGfxState or FVulkanPendingComputeState");
-	constexpr bool bIsGfx = TIsSame<TState, FVulkanPendingGfxState>::Value;
+	static_assert(std::is_same_v<TState, FVulkanPendingGfxState> || std::is_same_v<TState, FVulkanPendingComputeState>, "TState must be FVulkanPendingGfxState or FVulkanPendingComputeState");
+	constexpr bool bIsGfx = std::is_same_v<TState, FVulkanPendingGfxState>;
 	constexpr ERHIAccess SRVAccess = bIsGfx ? ERHIAccess::SRVGraphics : ERHIAccess::SRVCompute;
 	constexpr ERHIAccess UAVAccess = bIsGfx ? ERHIAccess::UAVGraphics : ERHIAccess::UAVCompute;
 #endif
diff --git a/Engine/Source/Runtime/Windows/D3D11RHI/Public/D3D11Util.h b/Engine/Source/Runtime/Windows/D3D11RHI/Public/D3D11Util.h
--- a/Engine/Source/Runtime/Windows/D3D11RHI/Public/D3D11Util.h
+++ b/Engine/Source/Runtime/Windows/D3D11RHI/Public/D3D11Util.h
@@ -241,10 +241,10 @@ struct FD3D11RHIGenericCommandString
 };
 template <
 	typename JobType,
-	typename = TEnableIf<TOr<
-	TIsSame<JobType, TFunction<void()>>,
-	TIsSame<JobType, TFunction<void()>&>>::Value>>
-	class TD3D11RHIGenericCommand final : public FRHICommand<TD3D11RHIGenericCommand<JobType>, FD3D11RHIGenericCommandString>
+	typename = TEnableIf<
+	std::is_same_v<JobType, TFunction<void()>> ||
+	std::is_same_v<JobType, TFunction<void()>&>>>
+class TD3D11RHIGenericCommand final : public FRHICommand<TD3D11RHIGenericCommand<JobType>, FD3D11RHIGenericCommandString>
 {
 public:
 	// InRHIJob is supposed to be called on RHIT (don't capture things that can become outdated here)
