Code cleanup to reduce the amount of places calling FEditorBulkData::GetPayloadSize

#rnx

- Ultimately I'd like to be able to remove FEditorBulkData::GetPayloadSize all together but there are a few places that use the byte length to make estimates about memory usage.
- To move towards this goal, this change removes a number of calls to ::GetPayloadSize which only check if the length is or isn't zero, which can be accomplished with a call to ::HasPayloadData instead.

diff --git a/Engine/Source/Runtime/MeshDescription/Private/MeshDescription.cpp b/Engine/Source/Runtime/MeshDescription/Private/MeshDescription.cpp
--- a/Engine/Source/Runtime/MeshDescription/Private/MeshDescription.cpp
+++ b/Engine/Source/Runtime/MeshDescription/Private/MeshDescription.cpp
@@ -2114,7 +2114,7 @@ void FMeshDescriptionBulkData::LoadMeshDescription( FMeshDescription& MeshDescri
 
 	MeshDescription.Empty();
 
-	if (BulkData.GetPayloadSize() > 0)
+	if (BulkData.HasPayloadData())
 	{
 #if WITH_EDITOR
 		// TODO: Remove this once VirtualizedBulkData can be shown to be thread safe
@@ -2143,7 +2143,8 @@ void FMeshDescriptionBulkData::UpdateMeshDescriptionFormat()
 	{
 		return;
 	}
-	if (BulkData.GetPayloadSize() == 0)
+
+	if (!BulkData.HasPayloadData())
 	{
 		return;
 	}
@@ -2230,7 +2231,7 @@ FString FMeshDescriptionBulkData::GetIdString() const
 
 FGuid FMeshDescriptionBulkData::GetHash() const
 {
-	if (BulkData.GetPayloadSize() > 0)
+	if (BulkData.HasPayloadData())
 	{
 		return UE::Serialization::IoHashToGuid(BulkData.GetPayloadId());
 	}
@@ -2240,7 +2241,7 @@ FGuid FMeshDescriptionBulkData::GetHash() const
 
 void FMeshDescriptionBulkData::UseHashAsGuid()
 {
-	if (BulkData.GetPayloadSize() > 0)
+	if (BulkData.HasPayloadData())
 	{
 		bGuidIsHash = true;
 
diff --git a/Engine/Source/Runtime/MeshDescription/Public/MeshDescription.h b/Engine/Source/Runtime/MeshDescription/Public/MeshDescription.h
--- a/Engine/Source/Runtime/MeshDescription/Public/MeshDescription.h
+++ b/Engine/Source/Runtime/MeshDescription/Public/MeshDescription.h
@@ -1336,7 +1336,7 @@ public:
 	void Empty();
 
 	/** Returns true if there is no bulk data available */
-	bool IsEmpty() const { return BulkData.GetPayloadSize() == 0; }
+	bool IsEmpty() const { return !BulkData.HasPayloadData(); }
 
 	/** Return unique ID string for this bulk data */
 	FString GetIdString() const;
