Removed TRemoveCV usage.

diff --git a/Engine/Source/Runtime/Core/Public/Containers/ArrayView.h b/Engine/Source/Runtime/Core/Public/Containers/ArrayView.h
--- a/Engine/Source/Runtime/Core/Public/Containers/ArrayView.h
+++ b/Engine/Source/Runtime/Core/Public/Containers/ArrayView.h
@@ -160,7 +160,7 @@ public:
 	 */
 	template <
 		typename OtherRangeType,
-		typename CVUnqualifiedOtherRangeType = typename TRemoveCV<typename TRemoveReference<OtherRangeType>::Type>::Type,
+		typename CVUnqualifiedOtherRangeType = std::remove_cv_t<typename TRemoveReference<OtherRangeType>::Type>,
 		typename = typename TEnableIf<
 			TAnd<
 				TIsContiguousContainer<CVUnqualifiedOtherRangeType>,
@@ -185,7 +185,7 @@ public:
 	}
 	template <
 		typename OtherRangeType,
-		typename CVUnqualifiedOtherRangeType = typename TRemoveCV<typename TRemoveReference<OtherRangeType>::Type>::Type,
+		typename CVUnqualifiedOtherRangeType = std::remove_cv_t<typename TRemoveReference<OtherRangeType>::Type>,
 		typename = typename TEnableIf<
 		TAnd<
 		TIsContiguousContainer<CVUnqualifiedOtherRangeType>,
@@ -749,7 +749,7 @@ struct TIsContiguousContainer<TArrayView<T, SizeType>>
 
 template <
 	typename OtherRangeType,
-	typename CVUnqualifiedOtherRangeType = typename TRemoveCV<typename TRemoveReference<OtherRangeType>::Type>::Type,
+	typename CVUnqualifiedOtherRangeType = std::remove_cv_t<typename TRemoveReference<OtherRangeType>::Type>,
 	typename = typename TEnableIf<TIsContiguousContainer<CVUnqualifiedOtherRangeType>::Value>::Type,
 	std::enable_if_t<TIsTArrayView<std::decay_t<OtherRangeType>>::Value>* = nullptr
 >
@@ -759,7 +759,7 @@ auto MakeArrayView(OtherRangeType&& Other)
 }
 template <
 	typename OtherRangeType,
-	typename CVUnqualifiedOtherRangeType = typename TRemoveCV<typename TRemoveReference<OtherRangeType>::Type>::Type,
+	typename CVUnqualifiedOtherRangeType = std::remove_cv_t<typename TRemoveReference<OtherRangeType>::Type>,
 	typename = typename TEnableIf<TIsContiguousContainer<CVUnqualifiedOtherRangeType>::Value>::Type,
 	std::enable_if_t<!TIsTArrayView<std::decay_t<OtherRangeType>>::Value>* = nullptr
 >
diff --git a/Engine/Source/Runtime/Core/Public/Containers/UnrealString.h b/Engine/Source/Runtime/Core/Public/Containers/UnrealString.h
--- a/Engine/Source/Runtime/Core/Public/Containers/UnrealString.h
+++ b/Engine/Source/Runtime/Core/Public/Containers/UnrealString.h
@@ -62,7 +62,7 @@ private:
 	DataType Data;
 
 	template <typename RangeType>
-	using TRangeElementType = typename TRemoveCV<typename TRemovePointer<decltype(GetData(DeclVal<RangeType>()))>::Type>::Type;
+	using TRangeElementType = std::remove_cv_t<typename TRemovePointer<decltype(GetData(DeclVal<RangeType>()))>::Type>;
 
 	template <typename CharRangeType>
 	struct TIsRangeOfCharType : TIsCharType<TRangeElementType<CharRangeType>>
@@ -2246,8 +2246,8 @@ template<typename T>
 UE_NODISCARD typename TEnableIf<TIsArithmetic<T>::Value, FString>::Type
 LexToString(const T& Value)
 {
-	// TRemoveCV to remove potential volatile decorations. Removing const is pointless, but harmless because it's specified in the param declaration.
-	return FString::Printf( TFormatSpecifier<typename TRemoveCV<T>::Type>::GetFormatSpecifier(), Value );
+	// std::remove_cv_t to remove potential volatile decorations. Removing const is pointless, but harmless because it's specified in the param declaration.
+	return FString::Printf( TFormatSpecifier<std::remove_cv_t<T>>::GetFormatSpecifier(), Value );
 }
 
 template<typename CharType>
diff --git a/Engine/Source/Runtime/Core/Public/ProfilingDebugging/FormatArgsTrace.h b/Engine/Source/Runtime/Core/Public/ProfilingDebugging/FormatArgsTrace.h
--- a/Engine/Source/Runtime/Core/Public/ProfilingDebugging/FormatArgsTrace.h
+++ b/Engine/Source/Runtime/Core/Public/ProfilingDebugging/FormatArgsTrace.h
@@ -49,7 +49,7 @@ private:
 	template <typename T>
 	struct TIsStringArgument
 	{
-		enum { Value = TAnd<TIsPointer<T>, TIsCharType<typename TRemoveCV<typename TRemovePointer<T>::Type>::Type>>::Value };
+		enum { Value = TAnd<TIsPointer<T>, TIsCharType<std::remove_cv_t<typename TRemovePointer<T>::Type>>>::Value };
 	};
 
 	template <typename T>
@@ -58,7 +58,7 @@ private:
 		return sizeof(T);
 	}
 
-	template <typename T, typename CharType = typename TRemoveCV<typename TRemovePointer<T>::Type>::Type>
+	template <typename T, typename CharType = std::remove_cv_t<typename TRemovePointer<T>::Type>>
 	static typename TEnableIf<TIsStringArgument<T>::Value, uint64>::Type GetArgumentEncodedSize(T Argument)
 	{
 		if (Argument != nullptr)
@@ -116,7 +116,7 @@ private:
 		PayloadPtr += sizeof(T);
 	}
 
-	template <typename T, typename CharType = typename TRemoveCV<typename TRemovePointer<T>::Type>::Type>
+	template <typename T, typename CharType = std::remove_cv_t<typename TRemovePointer<T>::Type>>
 	static typename TEnableIf<TIsStringArgument<T>::Value>::Type EncodeArgumentInternal(uint8*& TypeCodesPtr, uint8*& PayloadPtr, T Argument)
 	{
 		*TypeCodesPtr++ = FormatArgTypeCode_CategoryString | sizeof(CharType);
diff --git a/Engine/Source/Runtime/Core/Public/Templates/Decay.h b/Engine/Source/Runtime/Core/Public/Templates/Decay.h
--- a/Engine/Source/Runtime/Core/Public/Templates/Decay.h
+++ b/Engine/Source/Runtime/Core/Public/Templates/Decay.h
@@ -3,15 +3,15 @@
 #pragma once
 
 #include "CoreTypes.h"
-#include "Templates/RemoveCV.h"
 #include "Templates/RemoveReference.h"
+#include <type_traits>
 
 namespace UE::Core::Private::Decay
 {
 	template <typename T>
 	struct TDecayNonReference
 	{
-		typedef typename TRemoveCV<T>::Type Type;
+		using Type = std::remove_cv_t<T>;
 	};
 
 	template <typename T>
diff --git a/Engine/Source/Runtime/Core/Public/Templates/PointerIsConvertibleFromTo.h b/Engine/Source/Runtime/Core/Public/Templates/PointerIsConvertibleFromTo.h
--- a/Engine/Source/Runtime/Core/Public/Templates/PointerIsConvertibleFromTo.h
+++ b/Engine/Source/Runtime/Core/Public/Templates/PointerIsConvertibleFromTo.h
@@ -4,12 +4,12 @@
 
 #include "CoreTypes.h"
 #include "Misc/StaticAssertCompleteType.h"
-#include "Templates/RemoveCV.h"
 #include "Templates/LosesQualifiersFromTo.h"
+#include <type_traits>
 
 namespace UE::Core::Private::PointerIsConvertibleFromTo
 {
-	template <typename From, typename To, typename NoCVFrom = typename TRemoveCV<From>::Type, typename NoCVTo = typename TRemoveCV<To>::Type>
+	template <typename From, typename To, typename NoCVFrom = std::remove_cv_t<From>, typename NoCVTo = std::remove_cv_t<To>>
 	struct TImpl
 	{
 	private:
diff --git a/Engine/Source/Runtime/Core/Public/Templates/RemoveCV.h b/Engine/Source/Runtime/Core/Public/Templates/RemoveCV.h
--- a/Engine/Source/Runtime/Core/Public/Templates/RemoveCV.h
+++ b/Engine/Source/Runtime/Core/Public/Templates/RemoveCV.h
@@ -3,13 +3,15 @@
 #pragma once
 
 #include "CoreTypes.h"
+#include <type_traits>
 
 /**
  * TRemoveCV<type> will remove any const/volatile qualifiers from a type.
  * (based on std::remove_cv<>
  * note: won't remove the const from "const int*", as the pointer is not const
  */
-template <typename T> struct TRemoveCV                   { typedef T Type; };
-template <typename T> struct TRemoveCV<const T>          { typedef T Type; };
-template <typename T> struct TRemoveCV<volatile T>       { typedef T Type; };
-template <typename T> struct TRemoveCV<const volatile T> { typedef T Type; };
+template <typename T>
+struct /*UE_DEPRECATED(5.3, "TRemoveCV has been deprecated, please use std::remove_cv_t instead.")*/ TRemoveCV
+{
+	using Type = std::remove_cv_t<T>;
+};
diff --git a/Engine/Source/Runtime/Core/Public/Templates/UnrealTypeTraits.h b/Engine/Source/Runtime/Core/Public/Templates/UnrealTypeTraits.h
--- a/Engine/Source/Runtime/Core/Public/Templates/UnrealTypeTraits.h
+++ b/Engine/Source/Runtime/Core/Public/Templates/UnrealTypeTraits.h
@@ -476,8 +476,8 @@ struct TIsBitwiseConstructible
 		"TIsBitwiseConstructible is not designed to accept reference types");
 
 	static_assert(
-		std::is_same_v<T,   typename TRemoveCV<T  >::Type> &&
-		std::is_same_v<Arg, typename TRemoveCV<Arg>::Type>,
+		std::is_same_v<T,   std::remove_cv_t<T  >> &&
+		std::is_same_v<Arg, std::remove_cv_t<Arg>>,
 		"TIsBitwiseConstructible is not designed to accept qualified types");
 
 	// Assume no bitwise construction in general
diff --git a/Engine/Source/Runtime/MeshDescription/Public/AttributeArrayContainer.h b/Engine/Source/Runtime/MeshDescription/Public/AttributeArrayContainer.h
--- a/Engine/Source/Runtime/MeshDescription/Public/AttributeArrayContainer.h
+++ b/Engine/Source/Runtime/MeshDescription/Public/AttributeArrayContainer.h
@@ -392,7 +392,7 @@ template <typename AttributeType>
 class TArrayAttribute
 {
 	template <typename T> friend class TArrayAttribute;
-	using ArrayType = typename TCopyQualifiersFromTo<AttributeType, TAttributeArrayContainer<typename TRemoveCV<AttributeType>::Type>>::Type;
+	using ArrayType = typename TCopyQualifiersFromTo<AttributeType, TAttributeArrayContainer<std::remove_cv_t<AttributeType>>>::Type;
 
 public:
 	explicit TArrayAttribute(ArrayType& InArray, int32 InIndex)
@@ -404,7 +404,7 @@ public:
 	 * Construct a TArrayAttribute<const T> from a TArrayAttribute<T>. 
 	 */
 	template <typename T = AttributeType, typename TEnableIf<std::is_same_v<T, const T>, int>::Type = 0>
-	TArrayAttribute(TArrayAttribute<typename TRemoveCV<T>::Type> InValue)
+	TArrayAttribute(TArrayAttribute<std::remove_cv_t<T>> InValue)
 		: Array(InValue.Array),
 		  Index(InValue.Index)
 	{}
diff --git a/Engine/Source/Runtime/MeshDescription/Public/MeshAttributeArray.h b/Engine/Source/Runtime/MeshDescription/Public/MeshAttributeArray.h
--- a/Engine/Source/Runtime/MeshDescription/Public/MeshAttributeArray.h
+++ b/Engine/Source/Runtime/MeshDescription/Public/MeshAttributeArray.h
@@ -729,7 +729,7 @@ struct TMeshAttributesRefType : TMeshAttributesRefTypeBase<T>
 	static const uint32 MaxExpectedExtent = 1;
 	using RefType = T;
 	using ConstRefType = const T;
-	using NonConstRefType = typename TRemoveCV<T>::Type;
+	using NonConstRefType = std::remove_cv_t<T>;
 };
 
 template <typename T>
@@ -739,7 +739,7 @@ struct TMeshAttributesRefType<TArrayView<T>> : TMeshAttributesRefTypeBase<T>
 	static const uint32 MaxExpectedExtent = 0xFFFFFFFF;
 	using RefType = TArrayView<T>;
 	using ConstRefType = TArrayView<const T>;
-	using NonConstRefType = TArrayView<typename TRemoveCV<T>::Type>;
+	using NonConstRefType = TArrayView<std::remove_cv_t<T>>;
 };
 
 template <typename T>
@@ -749,7 +749,7 @@ struct TMeshAttributesRefType<TArrayAttribute<T>> : TMeshAttributesRefTypeBase<T
 	static const uint32 MaxExpectedExtent = 0;
 	using RefType = TArrayAttribute<T>;
 	using ConstRefType = TArrayAttribute<const T>;
-	using NonConstRefType = TArrayAttribute<typename TRemoveCV<T>::Type>;
+	using NonConstRefType = TArrayAttribute<std::remove_cv_t<T>>;
 };
 
 
@@ -808,7 +808,7 @@ class TMeshAttributesRef
 
 public:
 	using BaseArrayType = typename TCopyQualifiersFromTo<AttributeType, FMeshAttributeArraySetBase>::Type;
-	using ArrayType = typename TCopyQualifiersFromTo<AttributeType, TMeshAttributeArraySet<typename TRemoveCV<AttributeType>::Type>>::Type;
+	using ArrayType = typename TCopyQualifiersFromTo<AttributeType, TMeshAttributeArraySet<std::remove_cv_t<AttributeType>>>::Type;
 
 	/** Constructor taking a pointer to a FMeshAttributeArraySetBase */
 	explicit TMeshAttributesRef(BaseArrayType* InArrayPtr = nullptr, uint32 InExtent = 1)
@@ -817,7 +817,7 @@ public:
 
 	/** Implicitly construct a TMeshAttributesRef-to-const from a regular one */
 	template <typename T = AttributeType, typename TEnableIf<std::is_same_v<T, const T>, int>::Type = 0>
-	TMeshAttributesRef(TMeshAttributesRef<ElementIDType, typename TRemoveCV<T>::Type> InRef)
+	TMeshAttributesRef(TMeshAttributesRef<ElementIDType, std::remove_cv_t<T>> InRef)
 		: ArrayPtr(InRef.ArrayPtr)
 	{}
 
@@ -829,7 +829,7 @@ public:
 
 	/** Implicitly construct a TMeshAttributesRef-to-const from a TMeshAttributesArray */
 	template <typename T = ElementIDType, typename U = AttributeType, typename TEnableIf<std::is_same_v<U, const U>, int>::Type = 0, typename TEnableIf<!std::is_same_v<T, int32>, int>::Type = 0>
-	TMeshAttributesRef(TMeshAttributesRef<int32, typename TRemoveCV<U>::Type> InRef)
+	TMeshAttributesRef(TMeshAttributesRef<int32, std::remove_cv_t<U>> InRef)
 		: ArrayPtr(InRef.ArrayPtr)
 	{}
 
@@ -1003,8 +1003,8 @@ class TMeshAttributesRef<ElementIDType, TArrayView<AttributeType>>
 
 public:
 	using BaseArrayType = typename TCopyQualifiersFromTo<AttributeType, FMeshAttributeArraySetBase>::Type;
-	using BoundedArrayType = typename TCopyQualifiersFromTo<AttributeType, TMeshAttributeArraySet<typename TRemoveCV<AttributeType>::Type>>::Type;
-	using UnboundedArrayType = typename TCopyQualifiersFromTo<AttributeType, TMeshUnboundedAttributeArraySet<typename TRemoveCV<AttributeType>::Type>>::Type;
+	using BoundedArrayType = typename TCopyQualifiersFromTo<AttributeType, TMeshAttributeArraySet<std::remove_cv_t<AttributeType>>>::Type;
+	using UnboundedArrayType = typename TCopyQualifiersFromTo<AttributeType, TMeshUnboundedAttributeArraySet<std::remove_cv_t<AttributeType>>>::Type;
 
 	/** Constructor taking a pointer to a TMeshAttributeArraySet */
 	explicit TMeshAttributesRef(BaseArrayType* InArrayPtr = nullptr, uint32 InExtent = 1)
@@ -1014,7 +1014,7 @@ public:
 
 	/** Implicitly construct a TMeshAttributesRef-to-const from a regular one */
 	template <typename T = AttributeType, typename TEnableIf<std::is_same_v<T, const T>, int>::Type = 0>
-	TMeshAttributesRef(TMeshAttributesRef<ElementIDType, TArrayView<typename TRemoveCV<T>::Type>> InRef)
+	TMeshAttributesRef(TMeshAttributesRef<ElementIDType, TArrayView<std::remove_cv_t<T>>> InRef)
 		: ArrayPtr(InRef.ArrayPtr),
 		  Extent(InRef.Extent)
 	{}
@@ -1028,7 +1028,7 @@ public:
 
 	/** Implicitly construct a TMeshAttributesRef-to-const from a TMeshAttributesArray */
 	template <typename T = ElementIDType, typename U = AttributeType, typename TEnableIf<std::is_same_v<U, const U>, int>::Type = 0, typename TEnableIf<!std::is_same_v<T, int32>, int>::Type = 0>
-	TMeshAttributesRef(TMeshAttributesRef<int32, TArrayView<typename TRemoveCV<U>::Type>> InRef)
+	TMeshAttributesRef(TMeshAttributesRef<int32, TArrayView<std::remove_cv_t<U>>> InRef)
 		: ArrayPtr(InRef.ArrayPtr),
 		  Extent(InRef.Extent)
 	{}
@@ -1270,7 +1270,7 @@ class TMeshAttributesRef<ElementIDType, TArrayAttribute<AttributeType>>
 
 public:
 	using BaseArrayType = typename TCopyQualifiersFromTo<AttributeType, FMeshAttributeArraySetBase>::Type;
-	using ArrayType = typename TCopyQualifiersFromTo<AttributeType, TMeshUnboundedAttributeArraySet<typename TRemoveCV<AttributeType>::Type>>::Type;
+	using ArrayType = typename TCopyQualifiersFromTo<AttributeType, TMeshUnboundedAttributeArraySet<std::remove_cv_t<AttributeType>>>::Type;
 
 	/** Constructor taking a pointer to a TMeshUnboundedAttributeArraySet */
 	explicit TMeshAttributesRef(BaseArrayType* InArrayPtr = nullptr, uint32 InExtent = 0)
@@ -1279,7 +1279,7 @@ public:
 
 	/** Implicitly construct a TMeshAttributesRef-to-const from a regular one */
 	template <typename T = AttributeType, typename TEnableIf<std::is_same_v<T, const T>, int>::Type = 0>
-	TMeshAttributesRef(TMeshAttributesRef<ElementIDType, TArrayAttribute<typename TRemoveCV<T>::Type>> InRef)
+	TMeshAttributesRef(TMeshAttributesRef<ElementIDType, TArrayAttribute<std::remove_cv_t<T>>> InRef)
 		: ArrayPtr(InRef.ArrayPtr)
 	{}
 
@@ -1291,7 +1291,7 @@ public:
 
 	/** Implicitly construct a TMeshAttributesRef-to-const from a TMeshAttributesArray */
 	template <typename T = ElementIDType, typename U = AttributeType, typename TEnableIf<std::is_same_v<U, const U>, int>::Type = 0, typename TEnableIf<!std::is_same_v<T, int32>, int>::Type = 0>
-	TMeshAttributesRef(TMeshAttributesRef<int32, TArrayAttribute<typename TRemoveCV<U>::Type>> InRef)
+	TMeshAttributesRef(TMeshAttributesRef<int32, TArrayAttribute<std::remove_cv_t<U>>> InRef)
 		: ArrayPtr(InRef.ArrayPtr)
 	{}
 
@@ -2018,7 +2018,7 @@ public:
 
 		FMeshAttributeArraySetBase* ArraySetPtr = this->Map.FindChecked(AttributeName).Get();
 		uint32 ActualExtent = ArraySetPtr->GetExtent();
-		check(ArraySetPtr->HasType<typename TRemoveCV<RealAttributeType>::Type>());
+		check(ArraySetPtr->HasType<std::remove_cv_t<RealAttributeType>>());
 		check(ActualExtent >= TMeshAttributesRefType<AttributeType>::MinExpectedExtent && ActualExtent <= TMeshAttributesRefType<AttributeType>::MaxExpectedExtent);
 
 		TMeshAttributesRef<FElementID, NonConstRefType> Ref(ArraySetPtr, ActualExtent);
