MeshDescription: Explicit specialization for hashing FName attributes by their string value instead of comparison index and number.

diff --git a/Engine/Source/Runtime/MeshDescription/Public/MeshAttributeArray.h b/Engine/Source/Runtime/MeshDescription/Public/MeshAttributeArray.h
--- a/Engine/Source/Runtime/MeshDescription/Public/MeshAttributeArray.h
+++ b/Engine/Source/Runtime/MeshDescription/Public/MeshAttributeArray.h
@@ -162,6 +162,18 @@ protected:
 	uint32 Extent;
 };
 
+// We have to manually implement this or else the default behavior will just hash FName ComparisonIndex and Number,
+// which aren't deterministic
+template<>
+inline uint32 TMeshAttributeArrayBase<FName>::GetHash(uint32 Crc) const
+{
+	for (const FName& Item : Container)
+	{
+		const FString ItemStr = Item.ToString();
+		Crc = FCrc::MemCrc32(*ItemStr, ItemStr.Len() * sizeof(TCHAR), Crc);
+	}
+	return Crc;
+}
 
 template <typename AttributeType>
 void TMeshAttributeArrayBase<AttributeType>::Insert(const int32 Index, const AttributeType& Default)
