Added FTransform attribute to the list of supported attributes.

#rnx

diff --git a/Engine/Source/Runtime/MeshDescription/Private/MeshDescription.cpp b/Engine/Source/Runtime/MeshDescription/Private/MeshDescription.cpp
--- a/Engine/Source/Runtime/MeshDescription/Private/MeshDescription.cpp
+++ b/Engine/Source/Runtime/MeshDescription/Private/MeshDescription.cpp
@@ -379,6 +379,12 @@ struct FFixAttributesSizeHelper
 		check(false);
 	}
 
+	void operator()(const FName AttributeName, TMeshAttributesRef<T, FTransform> AttributeArrayRef)
+	{
+		// Not expecting FTransform in legacy attributes
+		check(false);
+	}
+
 	template <typename U>
 	void operator()(const FName AttributeName, TMeshAttributesRef<T, U> AttributeArrayRef)
 	{
diff --git a/Engine/Source/Runtime/MeshDescription/Public/MeshAttributeArray.h b/Engine/Source/Runtime/MeshDescription/Public/MeshAttributeArray.h
--- a/Engine/Source/Runtime/MeshDescription/Public/MeshAttributeArray.h
+++ b/Engine/Source/Runtime/MeshDescription/Public/MeshAttributeArray.h
@@ -53,7 +53,8 @@ using AttributeTypes = TTuple
 	float,
 	int32,
 	bool,
-	FName
+	FName,
+	FTransform
 >;
 
 
