Build error fix for latest clang release.

diff --git a/Engine/Source/Runtime/MeshDescription/Public/MeshAttributeArray.h b/Engine/Source/Runtime/MeshDescription/Public/MeshAttributeArray.h
--- a/Engine/Source/Runtime/MeshDescription/Public/MeshAttributeArray.h
+++ b/Engine/Source/Runtime/MeshDescription/Public/MeshAttributeArray.h
@@ -827,21 +827,35 @@ public:
 		: ArrayPtr(InArrayPtr)
 	{}
 
+	/** Basic copy constructor */
+	TMeshAttributesRef(const TMeshAttributesRef<ElementIDType, AttributeType>& InRef)
+		: ArrayPtr(InRef.ArrayPtr)
+	{}
+
 	/** Implicitly construct a TMeshAttributesRef-to-const from a regular one */
-	template <typename T = AttributeType, typename TEnableIf<std::is_same_v<T, const T>, int>::Type = 0>
-	TMeshAttributesRef(TMeshAttributesRef<ElementIDType, std::remove_cv_t<T>> InRef)
+	template <typename SrcAttributeType,
+			  typename DestAttributeType = AttributeType,
+			  typename TEnableIf<std::is_const_v<DestAttributeType>, int>::Type = 0,
+			  typename TEnableIf<!std::is_const_v<SrcAttributeType>, int>::Type = 0>
+	TMeshAttributesRef(const TMeshAttributesRef<ElementIDType, SrcAttributeType>& InRef)
 		: ArrayPtr(InRef.ArrayPtr)
 	{}
 
 	/** Implicitly construct a TMeshAttributesRef from a TMeshAttributesArray **/
-	template <typename T = ElementIDType, typename U = AttributeType, typename TEnableIf<!std::is_same_v<T, int32>, int>::Type = 0>
-	TMeshAttributesRef(TMeshAttributesRef<int32, U> InRef)
+	template <typename IDType = ElementIDType,
+			  typename TEnableIf<!std::is_same_v<IDType, int32>, int>::Type = 0>
+	TMeshAttributesRef(const TMeshAttributesRef<int32, AttributeType>& InRef)
 		: ArrayPtr(InRef.ArrayPtr)
 	{}
 
 	/** Implicitly construct a TMeshAttributesRef-to-const from a TMeshAttributesArray */
-	template <typename T = ElementIDType, typename U = AttributeType, typename TEnableIf<std::is_same_v<U, const U>, int>::Type = 0, typename TEnableIf<!std::is_same_v<T, int32>, int>::Type = 0>
-	TMeshAttributesRef(TMeshAttributesRef<int32, std::remove_cv_t<U>> InRef)
+	template <typename SrcAttributeType,
+			  typename DestAttributeType = AttributeType,
+			  typename IDType = ElementIDType,
+			  typename TEnableIf<!std::is_same_v<IDType, int32>, int>::Type = 0,
+			  typename TEnableIf<std::is_const_v<DestAttributeType>, int>::Type = 0,
+			  typename TEnableIf<!std::is_const_v<SrcAttributeType>, int>::Type = 0>
+	TMeshAttributesRef(const TMeshAttributesRef<int32, SrcAttributeType>& InRef)
 		: ArrayPtr(InRef.ArrayPtr)
 	{}
 
@@ -1024,23 +1038,38 @@ public:
 		  Extent(InExtent)
 	{}
 
+	/** Basic copy constructor */
+	TMeshAttributesRef(const TMeshAttributesRef<ElementIDType, TArrayView<AttributeType>>& InRef)
+		: ArrayPtr(InRef.ArrayPtr),
+		  Extent(InRef.Extent)
+	{}
+
 	/** Implicitly construct a TMeshAttributesRef-to-const from a regular one */
-	template <typename T = AttributeType, typename TEnableIf<std::is_same_v<T, const T>, int>::Type = 0>
-	TMeshAttributesRef(TMeshAttributesRef<ElementIDType, TArrayView<std::remove_cv_t<T>>> InRef)
+	template <typename SrcAttributeType,
+			  typename DestAttributeType = AttributeType,
+			  typename TEnableIf<std::is_const_v<DestAttributeType>, int>::Type = 0,
+			  typename TEnableIf<!std::is_const_v<SrcAttributeType>, int>::Type = 0>
+	TMeshAttributesRef(const TMeshAttributesRef<ElementIDType, TArrayView<SrcAttributeType>>& InRef)
 		: ArrayPtr(InRef.ArrayPtr),
 		  Extent(InRef.Extent)
 	{}
 
 	/** Implicitly construct a TMeshAttributesRef from a TMeshAttributesArray **/
-	template <typename T = ElementIDType, typename U = AttributeType, typename TEnableIf<!std::is_same_v<T, int32>, int>::Type = 0>
-	TMeshAttributesRef(TMeshAttributesRef<int32, TArrayView<U>> InRef)
+	template <typename IDType = ElementIDType,
+			  typename TEnableIf<!std::is_same_v<IDType, int32>, int>::Type = 0>
+	TMeshAttributesRef(const TMeshAttributesRef<int32, TArrayView<AttributeType>>& InRef)
 		: ArrayPtr(InRef.ArrayPtr),
 		  Extent(InRef.Extent)
 	{}
 
 	/** Implicitly construct a TMeshAttributesRef-to-const from a TMeshAttributesArray */
-	template <typename T = ElementIDType, typename U = AttributeType, typename TEnableIf<std::is_same_v<U, const U>, int>::Type = 0, typename TEnableIf<!std::is_same_v<T, int32>, int>::Type = 0>
-	TMeshAttributesRef(TMeshAttributesRef<int32, TArrayView<std::remove_cv_t<U>>> InRef)
+	template <typename SrcAttributeType,
+			  typename DestAttributeType = AttributeType,
+			  typename IDType = ElementIDType,
+			  typename TEnableIf<!std::is_same_v<IDType, int32>, int>::Type = 0,
+			  typename TEnableIf<std::is_const_v<DestAttributeType>, int>::Type = 0,
+			  typename TEnableIf<!std::is_const_v<SrcAttributeType>, int>::Type = 0>
+	TMeshAttributesRef(const TMeshAttributesRef<int32, TArrayView<SrcAttributeType>>& InRef)
 		: ArrayPtr(InRef.ArrayPtr),
 		  Extent(InRef.Extent)
 	{}
@@ -1289,21 +1318,35 @@ public:
 		: ArrayPtr(InArrayPtr)
 	{}
 
+	/** Basic copy constructor */
+	TMeshAttributesRef(const TMeshAttributesRef<ElementIDType, TArrayAttribute<AttributeType>>& InRef)
+		: ArrayPtr(InRef.ArrayPtr)
+	{}
+
 	/** Implicitly construct a TMeshAttributesRef-to-const from a regular one */
-	template <typename T = AttributeType, typename TEnableIf<std::is_same_v<T, const T>, int>::Type = 0>
-	TMeshAttributesRef(TMeshAttributesRef<ElementIDType, TArrayAttribute<std::remove_cv_t<T>>> InRef)
+	template <typename SrcAttributeType,
+			  typename DestAttributeType = AttributeType,
+			  typename TEnableIf<std::is_const_v<DestAttributeType>, int>::Type = 0,
+			  typename TEnableIf<!std::is_const_v<SrcAttributeType>, int>::Type = 0>
+	TMeshAttributesRef(const TMeshAttributesRef<ElementIDType, TArrayAttribute<SrcAttributeType>>& InRef)
 		: ArrayPtr(InRef.ArrayPtr)
 	{}
 
 	/** Implicitly construct a TMeshAttributesRef from a TMeshAttributesArray **/
-	template <typename T = ElementIDType, typename U = AttributeType, typename TEnableIf<!std::is_same_v<T, int32>, int>::Type = 0>
-	TMeshAttributesRef(TMeshAttributesRef<int32, TArrayAttribute<U>> InRef)
+	template <typename IDType = ElementIDType,
+			  typename TEnableIf<!std::is_same_v<IDType, int32>, int>::Type = 0>
+	TMeshAttributesRef(const TMeshAttributesRef<int32, TArrayAttribute<AttributeType>>& InRef)
 		: ArrayPtr(InRef.ArrayPtr)
 	{}
 
 	/** Implicitly construct a TMeshAttributesRef-to-const from a TMeshAttributesArray */
-	template <typename T = ElementIDType, typename U = AttributeType, typename TEnableIf<std::is_same_v<U, const U>, int>::Type = 0, typename TEnableIf<!std::is_same_v<T, int32>, int>::Type = 0>
-	TMeshAttributesRef(TMeshAttributesRef<int32, TArrayAttribute<std::remove_cv_t<U>>> InRef)
+	template <typename SrcAttributeType,
+			  typename DestAttributeType = AttributeType,
+			  typename IDType = ElementIDType,
+			  typename TEnableIf<!std::is_same_v<IDType, int32>, int>::Type = 0,
+			  typename TEnableIf<std::is_const_v<DestAttributeType>, int>::Type = 0,
+			  typename TEnableIf<!std::is_const_v<SrcAttributeType>, int>::Type = 0>
+	TMeshAttributesRef(const TMeshAttributesRef<int32, TArrayAttribute<SrcAttributeType>>& InRef)
 		: ArrayPtr(InRef.ArrayPtr)
 	{}
 
