Removed trivial explicit copy constructors from TMeshAttributesRef, to allow the compiler to generate assignment operators automatically.

diff --git a/Engine/Source/Runtime/MeshDescription/Public/MeshAttributeArray.h b/Engine/Source/Runtime/MeshDescription/Public/MeshAttributeArray.h
--- a/Engine/Source/Runtime/MeshDescription/Public/MeshAttributeArray.h
+++ b/Engine/Source/Runtime/MeshDescription/Public/MeshAttributeArray.h
@@ -827,11 +827,6 @@ public:
 		: ArrayPtr(InArrayPtr)
 	{}
 
-	/** Basic copy constructor */
-	TMeshAttributesRef(const TMeshAttributesRef<ElementIDType, AttributeType>& InRef)
-		: ArrayPtr(InRef.ArrayPtr)
-	{}
-
 	/** Implicitly construct a TMeshAttributesRef-to-const from a regular one */
 	template <typename SrcAttributeType,
 			  typename DestAttributeType = AttributeType,
@@ -1038,12 +1033,6 @@ public:
 		  Extent(InExtent)
 	{}
 
-	/** Basic copy constructor */
-	TMeshAttributesRef(const TMeshAttributesRef<ElementIDType, TArrayView<AttributeType>>& InRef)
-		: ArrayPtr(InRef.ArrayPtr),
-		  Extent(InRef.Extent)
-	{}
-
 	/** Implicitly construct a TMeshAttributesRef-to-const from a regular one */
 	template <typename SrcAttributeType,
 			  typename DestAttributeType = AttributeType,
@@ -1318,11 +1307,6 @@ public:
 		: ArrayPtr(InArrayPtr)
 	{}
 
-	/** Basic copy constructor */
-	TMeshAttributesRef(const TMeshAttributesRef<ElementIDType, TArrayAttribute<AttributeType>>& InRef)
-		: ArrayPtr(InRef.ArrayPtr)
-	{}
-
 	/** Implicitly construct a TMeshAttributesRef-to-const from a regular one */
 	template <typename SrcAttributeType,
 			  typename DestAttributeType = AttributeType,
