[Engine]

diff --git a/Engine/Source/Runtime/MeshDescription/Public/MeshDescriptionBase.h b/Engine/Source/Runtime/MeshDescription/Public/MeshDescriptionBase.h
--- a/Engine/Source/Runtime/MeshDescription/Public/MeshDescriptionBase.h
+++ b/Engine/Source/Runtime/MeshDescription/Public/MeshDescriptionBase.h
@@ -21,14 +21,14 @@ template <typename ElementIDType> class TAttributesSet;
 
 
 
-UCLASS(BlueprintType)
-class MESHDESCRIPTION_API UMeshDescriptionBase : public UObject
+UCLASS(BlueprintType, MinimalAPI)
+class UMeshDescriptionBase : public UObject
 {
 	GENERATED_BODY()
 
 public:
 	// UObject interface
-	virtual void Serialize(FArchive& Ar) override;
+	MESHDESCRIPTION_API virtual void Serialize(FArchive& Ar) override;
 	//
 
 	/** Get a reference to the actual mesh description */
@@ -49,13 +49,13 @@ public:
 	}
 
 	// UMeshDescriptionBase interface
-	virtual void RegisterAttributes();
+	MESHDESCRIPTION_API virtual void RegisterAttributes();
 	virtual FMeshAttributes& GetRequiredAttributes() { return *RequiredAttributes; }
 	virtual const FMeshAttributes& GetRequiredAttributes() const { return *RequiredAttributes; }
 	//
 
 	/** Reset the contained mesh description */
-	void Reset();
+	MESHDESCRIPTION_API void Reset();
 
 	/** Accessors for mesh element arrays */
 	FVertexArray& Vertices() { return GetMeshDescription().Vertices(); }
@@ -103,11 +103,11 @@ public:
 
 	/** Empty the mesh description */
 	UFUNCTION(BlueprintCallable, Category="MeshDescription")
-	void Empty();
+	MESHDESCRIPTION_API void Empty();
 
 	/** Return whether the mesh description is empty */
 	UFUNCTION(BlueprintPure, Category="MeshDescription")
-	bool IsEmpty() const;
+	MESHDESCRIPTION_API bool IsEmpty() const;
 
 
 	///////////////////////////////////////////////////////
@@ -115,23 +115,23 @@ public:
 
 	/** Reserves space for this number of new vertices */
 	UFUNCTION(BlueprintCallable, Category="MeshDescription")
-	void ReserveNewVertices(int32 NumberOfNewVertices);
+	MESHDESCRIPTION_API void ReserveNewVertices(int32 NumberOfNewVertices);
 
 	/** Adds a new vertex to the mesh and returns its ID */
 	UFUNCTION(BlueprintCallable, Category="MeshDescription")
-	FVertexID CreateVertex();
+	MESHDESCRIPTION_API FVertexID CreateVertex();
 
 	/** Adds a new vertex to the mesh with the given ID */
 	UFUNCTION(BlueprintCallable, Category="MeshDescription")
-	void CreateVertexWithID(FVertexID VertexID);
+	MESHDESCRIPTION_API void CreateVertexWithID(FVertexID VertexID);
 
 	/** Deletes a vertex from the mesh */
 	UFUNCTION(BlueprintCallable, Category="MeshDescription")
-	void DeleteVertex(FVertexID VertexID);
+	MESHDESCRIPTION_API void DeleteVertex(FVertexID VertexID);
 
 	/** Returns whether the passed vertex ID is valid */
 	UFUNCTION(BlueprintPure, Category="MeshDescription")
-	bool IsVertexValid(FVertexID VertexID) const;
+	MESHDESCRIPTION_API bool IsVertexValid(FVertexID VertexID) const;
 
 	/** Returns the number of vertices */
 	UFUNCTION(BlueprintPure, Category="MeshDescription")
@@ -139,23 +139,23 @@ public:
 
 	/** Reserves space for this number of new vertex instances */
 	UFUNCTION(BlueprintCallable, Category="MeshDescription")
-	void ReserveNewVertexInstances(int32 NumberOfNewVertexInstances);
+	MESHDESCRIPTION_API void ReserveNewVertexInstances(int32 NumberOfNewVertexInstances);
 
 	/** Adds a new vertex instance to the mesh and returns its ID */
 	UFUNCTION(BlueprintCallable, Category="MeshDescription")
-	FVertexInstanceID CreateVertexInstance(FVertexID VertexID);
+	MESHDESCRIPTION_API FVertexInstanceID CreateVertexInstance(FVertexID VertexID);
 
 	/** Adds a new vertex instance to the mesh with the given ID */
 	UFUNCTION(BlueprintCallable, Category="MeshDescription")
-	void CreateVertexInstanceWithID(FVertexInstanceID VertexInstanceID, FVertexID VertexID);
+	MESHDESCRIPTION_API void CreateVertexInstanceWithID(FVertexInstanceID VertexInstanceID, FVertexID VertexID);
 
 	/** Deletes a vertex instance from a mesh */
 	UFUNCTION(BlueprintCallable, Category="MeshDescription")
-	void DeleteVertexInstance(FVertexInstanceID VertexInstanceID, TArray<FVertexID>& OrphanedVertices);
+	MESHDESCRIPTION_API void DeleteVertexInstance(FVertexInstanceID VertexInstanceID, TArray<FVertexID>& OrphanedVertices);
 
 	/** Returns whether the passed vertex instance ID is valid */
 	UFUNCTION(BlueprintPure, Category="MeshDescription")
-	bool IsVertexInstanceValid(FVertexInstanceID VertexInstanceID) const;
+	MESHDESCRIPTION_API bool IsVertexInstanceValid(FVertexInstanceID VertexInstanceID) const;
 
 	/** Returns the number of vertex instances */
 	UFUNCTION(BlueprintPure, Category="MeshDescription")
@@ -163,23 +163,23 @@ public:
 
 	/** Reserves space for this number of new edges */
 	UFUNCTION(BlueprintCallable, Category="MeshDescription")
-	void ReserveNewEdges(int32 NumberOfNewEdges);
+	MESHDESCRIPTION_API void ReserveNewEdges(int32 NumberOfNewEdges);
 
 	/** Adds a new edge to the mesh and returns its ID */
 	UFUNCTION(BlueprintCallable, Category="MeshDescription")
-	FEdgeID CreateEdge(FVertexID VertexID0, FVertexID VertexID1);
+	MESHDESCRIPTION_API FEdgeID CreateEdge(FVertexID VertexID0, FVertexID VertexID1);
 
 	/** Adds a new edge to the mesh with the given ID */
 	UFUNCTION(BlueprintCallable, Category="MeshDescription")
-	void CreateEdgeWithID(FEdgeID EdgeID, FVertexID VertexID0, FVertexID VertexID1);
+	MESHDESCRIPTION_API void CreateEdgeWithID(FEdgeID EdgeID, FVertexID VertexID0, FVertexID VertexID1);
 
 	/** Deletes an edge from a mesh */
 	UFUNCTION(BlueprintCallable, Category="MeshDescription")
-	void DeleteEdge(FEdgeID EdgeID, TArray<FVertexID>& OrphanedVertices);
+	MESHDESCRIPTION_API void DeleteEdge(FEdgeID EdgeID, TArray<FVertexID>& OrphanedVertices);
 
 	/** Returns whether the passed edge ID is valid */
 	UFUNCTION(BlueprintPure, Category="MeshDescription")
-	bool IsEdgeValid(FEdgeID EdgeID) const;
+	MESHDESCRIPTION_API bool IsEdgeValid(FEdgeID EdgeID) const;
 
 	/** Returns the number of edges */
 	UFUNCTION(BlueprintPure, Category="MeshDescription")
@@ -187,23 +187,23 @@ public:
 
 	/** Reserves space for this number of new triangles */
 	UFUNCTION(BlueprintCallable, Category="MeshDescription")
-	void ReserveNewTriangles(int32 NumberOfNewTriangles);
+	MESHDESCRIPTION_API void ReserveNewTriangles(int32 NumberOfNewTriangles);
 
 	/** Adds a new triangle to the mesh and returns its ID. This will also make an encapsulating polygon, and any missing edges. */
 	UFUNCTION(BlueprintCallable, Category="MeshDescription")
-	FTriangleID CreateTriangle(FPolygonGroupID PolygonGroupID, const TArray<FVertexInstanceID>& VertexInstanceIDs, TArray<FEdgeID>& NewEdgeIDs);
+	MESHDESCRIPTION_API FTriangleID CreateTriangle(FPolygonGroupID PolygonGroupID, const TArray<FVertexInstanceID>& VertexInstanceIDs, TArray<FEdgeID>& NewEdgeIDs);
 
 	/** Adds a new triangle to the mesh with the given ID. This will also make an encapsulating polygon, and any missing edges. */
 	UFUNCTION(BlueprintCallable, Category="MeshDescription")
-	void CreateTriangleWithID(FTriangleID TriangleID, FPolygonGroupID PolygonGroupID, const TArray<FVertexInstanceID>& VertexInstanceIDs, TArray<FEdgeID>& NewEdgeIDs);
+	MESHDESCRIPTION_API void CreateTriangleWithID(FTriangleID TriangleID, FPolygonGroupID PolygonGroupID, const TArray<FVertexInstanceID>& VertexInstanceIDs, TArray<FEdgeID>& NewEdgeIDs);
 
 	/** Deletes a triangle from the mesh */
 	UFUNCTION(BlueprintCallable, Category="MeshDescription")
-	void DeleteTriangle(FTriangleID TriangleID, TArray<FEdgeID>& OrphanedEdges, TArray<FVertexInstanceID>& OrphanedVertexInstances, TArray<FPolygonGroupID>& OrphanedPolygonGroupsPtr);
+	MESHDESCRIPTION_API void DeleteTriangle(FTriangleID TriangleID, TArray<FEdgeID>& OrphanedEdges, TArray<FVertexInstanceID>& OrphanedVertexInstances, TArray<FPolygonGroupID>& OrphanedPolygonGroupsPtr);
 
 	/** Returns whether the passed triangle ID is valid */
 	UFUNCTION(BlueprintPure, Category="MeshDescription")
-	bool IsTriangleValid(const FTriangleID TriangleID) const;
+	MESHDESCRIPTION_API bool IsTriangleValid(const FTriangleID TriangleID) const;
 
 	/** Returns the number of triangles */
 	UFUNCTION(BlueprintPure, Category="MeshDescription")
@@ -211,23 +211,23 @@ public:
 
 	/** Reserves space for this number of new polygons */
 	UFUNCTION(BlueprintCallable, Category="MeshDescription")
-	void ReserveNewPolygons(const int32 NumberOfNewPolygons);
+	MESHDESCRIPTION_API void ReserveNewPolygons(const int32 NumberOfNewPolygons);
 
 	/** Adds a new polygon to the mesh and returns its ID. This will also make any missing edges, and all constituent triangles. */
 	UFUNCTION(BlueprintCallable, Category="MeshDescription")
-	FPolygonID CreatePolygon(FPolygonGroupID PolygonGroupID, TArray<FVertexInstanceID>& VertexInstanceIDs, TArray<FEdgeID>& NewEdgeIDs);
+	MESHDESCRIPTION_API FPolygonID CreatePolygon(FPolygonGroupID PolygonGroupID, TArray<FVertexInstanceID>& VertexInstanceIDs, TArray<FEdgeID>& NewEdgeIDs);
 
 	/** Adds a new polygon to the mesh with the given ID. This will also make any missing edges, and all constituent triangles. */
 	UFUNCTION(BlueprintCallable, Category="MeshDescription")
-	void CreatePolygonWithID(FPolygonID PolygonID, FPolygonGroupID PolygonGroupID, TArray<FVertexInstanceID>& VertexInstanceIDs, TArray<FEdgeID>& NewEdgeIDs);
+	MESHDESCRIPTION_API void CreatePolygonWithID(FPolygonID PolygonID, FPolygonGroupID PolygonGroupID, TArray<FVertexInstanceID>& VertexInstanceIDs, TArray<FEdgeID>& NewEdgeIDs);
 
 	/** Deletes a polygon from the mesh */
 	UFUNCTION(BlueprintCallable, Category="MeshDescription")
-	void DeletePolygon(FPolygonID PolygonID, TArray<FEdgeID>& OrphanedEdges, TArray<FVertexInstanceID>& OrphanedVertexInstances, TArray<FPolygonGroupID>& OrphanedPolygonGroups);
+	MESHDESCRIPTION_API void DeletePolygon(FPolygonID PolygonID, TArray<FEdgeID>& OrphanedEdges, TArray<FVertexInstanceID>& OrphanedVertexInstances, TArray<FPolygonGroupID>& OrphanedPolygonGroups);
 
 	/** Returns whether the passed polygon ID is valid */
 	UFUNCTION(BlueprintPure, Category="MeshDescription")
-	bool IsPolygonValid(FPolygonID PolygonID) const;
+	MESHDESCRIPTION_API bool IsPolygonValid(FPolygonID PolygonID) const;
 
 	/** Returns the number of polygons */
 	UFUNCTION(BlueprintPure, Category="MeshDescription")
@@ -235,23 +235,23 @@ public:
 
 	/** Reserves space for this number of new polygon groups */
 	UFUNCTION(BlueprintCallable, Category="MeshDescription")
-	void ReserveNewPolygonGroups(int32 NumberOfNewPolygonGroups);
+	MESHDESCRIPTION_API void ReserveNewPolygonGroups(int32 NumberOfNewPolygonGroups);
 
 	/** Adds a new polygon group to the mesh and returns its ID */
 	UFUNCTION(BlueprintCallable, Category="MeshDescription")
-	FPolygonGroupID CreatePolygonGroup();
+	MESHDESCRIPTION_API FPolygonGroupID CreatePolygonGroup();
 
 	/** Adds a new polygon group to the mesh with the given ID */
 	UFUNCTION(BlueprintCallable, Category="MeshDescription")
-	void CreatePolygonGroupWithID(FPolygonGroupID PolygonGroupID);
+	MESHDESCRIPTION_API void CreatePolygonGroupWithID(FPolygonGroupID PolygonGroupID);
 
 	/** Deletes a polygon group from the mesh */
 	UFUNCTION(BlueprintCallable, Category="MeshDescription")
-	void DeletePolygonGroup(FPolygonGroupID PolygonGroupID);
+	MESHDESCRIPTION_API void DeletePolygonGroup(FPolygonGroupID PolygonGroupID);
 
 	/** Returns whether the passed polygon group ID is valid */
 	UFUNCTION(BlueprintPure, Category="MeshDescription")
-	bool IsPolygonGroupValid(FPolygonGroupID PolygonGroupID) const;
+	MESHDESCRIPTION_API bool IsPolygonGroupValid(FPolygonGroupID PolygonGroupID) const;
 
 	/** Returns the number of polygon groups */
 	UFUNCTION(BlueprintPure, Category="MeshDescription")
@@ -263,55 +263,55 @@ public:
 
 	/** Returns whether a given vertex is orphaned, i.e. it doesn't form part of any polygon */
 	UFUNCTION(BlueprintPure, Category="MeshDescription")
-	bool IsVertexOrphaned(FVertexID VertexID) const;
+	MESHDESCRIPTION_API bool IsVertexOrphaned(FVertexID VertexID) const;
 
 	/** Returns the edge ID defined by the two given vertex IDs, if there is one; otherwise INDEX_NONE */
 	UFUNCTION(BlueprintPure, Category="MeshDescription")
-	FEdgeID GetVertexPairEdge(FVertexID VertexID0, FVertexID VertexID1) const;
+	MESHDESCRIPTION_API FEdgeID GetVertexPairEdge(FVertexID VertexID0, FVertexID VertexID1) const;
 
 	/** Returns reference to an array of Edge IDs connected to this vertex */
 	UFUNCTION(BlueprintPure, Category="MeshDescription")
-	void GetVertexConnectedEdges(FVertexID VertexID, TArray<FEdgeID>& OutEdgeIDs) const;
+	MESHDESCRIPTION_API void GetVertexConnectedEdges(FVertexID VertexID, TArray<FEdgeID>& OutEdgeIDs) const;
 
 	/** Returns number of edges connected to this vertex */
 	UFUNCTION(BlueprintPure, Category="MeshDescription")
-	int32 GetNumVertexConnectedEdges(FVertexID VertexID) const;
+	MESHDESCRIPTION_API int32 GetNumVertexConnectedEdges(FVertexID VertexID) const;
 
 	/** Returns reference to an array of VertexInstance IDs instanced from this vertex */
 	UFUNCTION(BlueprintPure, Category="MeshDescription")
-	void GetVertexVertexInstances(FVertexID VertexID, TArray<FVertexInstanceID>& OutVertexInstanceIDs) const;
+	MESHDESCRIPTION_API void GetVertexVertexInstances(FVertexID VertexID, TArray<FVertexInstanceID>& OutVertexInstanceIDs) const;
 
 	/** Returns number of vertex instances created from this vertex */
 	UFUNCTION(BlueprintPure, Category="MeshDescription")
-	int32 GetNumVertexVertexInstances(FVertexID VertexID) const;
+	MESHDESCRIPTION_API int32 GetNumVertexVertexInstances(FVertexID VertexID) const;
 
 	/** Returns the triangles connected to this vertex */
 	UFUNCTION(BlueprintPure, Category="MeshDescription")
-	void GetVertexConnectedTriangles(FVertexID VertexID, TArray<FTriangleID>& OutConnectedTriangleIDs) const;
+	MESHDESCRIPTION_API void GetVertexConnectedTriangles(FVertexID VertexID, TArray<FTriangleID>& OutConnectedTriangleIDs) const;
 
 	/** Returns number of triangles connected to this vertex */
 	UFUNCTION(BlueprintPure, Category="MeshDescription")
-	int32 GetNumVertexConnectedTriangles(FVertexID VertexID) const;
+	MESHDESCRIPTION_API int32 GetNumVertexConnectedTriangles(FVertexID VertexID) const;
 
 	/** Returns the polygons connected to this vertex */
 	UFUNCTION(BlueprintPure, Category="MeshDescription")
-	void GetVertexConnectedPolygons(FVertexID VertexID, TArray<FPolygonID>& OutConnectedPolygonIDs) const;
+	MESHDESCRIPTION_API void GetVertexConnectedPolygons(FVertexID VertexID, TArray<FPolygonID>& OutConnectedPolygonIDs) const;
 
 	/** Returns the number of polygons connected to this vertex */
 	UFUNCTION(BlueprintPure, Category="MeshDescription")
-	int32 GetNumVertexConnectedPolygons(FVertexID VertexID) const;
+	MESHDESCRIPTION_API int32 GetNumVertexConnectedPolygons(FVertexID VertexID) const;
 
 	/** Returns the vertices adjacent to this vertex */
 	UFUNCTION(BlueprintPure, Category="MeshDescription")
-	void GetVertexAdjacentVertices(FVertexID VertexID, TArray<FVertexID>& OutAdjacentVertexIDs) const;
+	MESHDESCRIPTION_API void GetVertexAdjacentVertices(FVertexID VertexID, TArray<FVertexID>& OutAdjacentVertexIDs) const;
 
 	/** Gets a vertex position */
 	UFUNCTION(BlueprintPure, Category="MeshDescription")
-	FVector GetVertexPosition(FVertexID VertexID) const;
+	MESHDESCRIPTION_API FVector GetVertexPosition(FVertexID VertexID) const;
 
 	/** Sets a vertex position */
 	UFUNCTION(BlueprintCallable, Category="MeshDescription")
-	void SetVertexPosition(FVertexID VertexID, const FVector& Position);
+	MESHDESCRIPTION_API void SetVertexPosition(FVertexID VertexID, const FVector& Position);
 
 
 	//////////////////////////////////////////////////////////////////////
@@ -319,27 +319,27 @@ public:
 
 	/** Returns the vertex ID associated with the given vertex instance */
 	UFUNCTION(BlueprintPure, Category="MeshDescription")
-	FVertexID GetVertexInstanceVertex(FVertexInstanceID VertexInstanceID) const;
+	MESHDESCRIPTION_API FVertexID GetVertexInstanceVertex(FVertexInstanceID VertexInstanceID) const;
 
 	/** Returns the edge ID defined by the two given vertex instance IDs, if there is one; otherwise INDEX_NONE */
 	UFUNCTION(BlueprintPure, Category="MeshDescription")
-	FEdgeID GetVertexInstancePairEdge(FVertexInstanceID VertexInstanceID0, FVertexInstanceID VertexInstanceID1) const;
+	MESHDESCRIPTION_API FEdgeID GetVertexInstancePairEdge(FVertexInstanceID VertexInstanceID0, FVertexInstanceID VertexInstanceID1) const;
 
 	/** Returns reference to an array of Triangle IDs connected to this vertex instance */
 	UFUNCTION(BlueprintPure, Category="MeshDescription")
-	void GetVertexInstanceConnectedTriangles(FVertexInstanceID VertexInstanceID, TArray<FTriangleID>& OutConnectedTriangleIDs) const;
+	MESHDESCRIPTION_API void GetVertexInstanceConnectedTriangles(FVertexInstanceID VertexInstanceID, TArray<FTriangleID>& OutConnectedTriangleIDs) const;
 
 	/** Returns the number of triangles connected to this vertex instance */
 	UFUNCTION(BlueprintPure, Category="MeshDescription")
-	int32 GetNumVertexInstanceConnectedTriangles(FVertexInstanceID VertexInstanceID) const;
+	MESHDESCRIPTION_API int32 GetNumVertexInstanceConnectedTriangles(FVertexInstanceID VertexInstanceID) const;
 
 	/** Returns the polygons connected to this vertex instance */
 	UFUNCTION(BlueprintPure, Category="MeshDescription")
-	void GetVertexInstanceConnectedPolygons(FVertexInstanceID VertexInstanceID, TArray<FPolygonID>& OutConnectedPolygonIDs) const;
+	MESHDESCRIPTION_API void GetVertexInstanceConnectedPolygons(FVertexInstanceID VertexInstanceID, TArray<FPolygonID>& OutConnectedPolygonIDs) const;
 
 	/** Returns the number of polygons connected to this vertex instance. */
 	UFUNCTION(BlueprintPure, Category="MeshDescription")
-	int32 GetNumVertexInstanceConnectedPolygons(FVertexInstanceID VertexInstanceID) const;
+	MESHDESCRIPTION_API int32 GetNumVertexInstanceConnectedPolygons(FVertexInstanceID VertexInstanceID) const;
 
 
 	//////////////////////////////////////////////////////////////////////
@@ -347,35 +347,35 @@ public:
 
 	/** Determine whether a given edge is an internal edge between triangles of a polygon */
 	UFUNCTION(BlueprintPure, Category="MeshDescription")
-	bool IsEdgeInternal(FEdgeID EdgeID) const;
+	MESHDESCRIPTION_API bool IsEdgeInternal(FEdgeID EdgeID) const;
 
 	/** Determine whether a given edge is an internal edge between triangles of a specific polygon */
 	UFUNCTION(BlueprintPure, Category="MeshDescription")
-	bool IsEdgeInternalToPolygon(FEdgeID EdgeID, FPolygonID PolygonID) const;
+	MESHDESCRIPTION_API bool IsEdgeInternalToPolygon(FEdgeID EdgeID, FPolygonID PolygonID) const;
 
 	/** Returns reference to an array of triangle IDs connected to this edge */
 	UFUNCTION(BlueprintPure, Category="MeshDescription")
-	void GetEdgeConnectedTriangles(FEdgeID EdgeID, TArray<FTriangleID>& OutConnectedTriangleIDs) const;
+	MESHDESCRIPTION_API void GetEdgeConnectedTriangles(FEdgeID EdgeID, TArray<FTriangleID>& OutConnectedTriangleIDs) const;
 
 	/** Returns the number of triangles connected to this edge */
 	UFUNCTION(BlueprintPure, Category="MeshDescription")
-	int32 GetNumEdgeConnectedTriangles(FEdgeID EdgeID) const;
+	MESHDESCRIPTION_API int32 GetNumEdgeConnectedTriangles(FEdgeID EdgeID) const;
 
 	/** Returns the polygons connected to this edge */
 	UFUNCTION(BlueprintPure, Category="MeshDescription")
-	void GetEdgeConnectedPolygons(FEdgeID EdgeID, TArray<FPolygonID>& OutConnectedPolygonIDs) const;
+	MESHDESCRIPTION_API void GetEdgeConnectedPolygons(FEdgeID EdgeID, TArray<FPolygonID>& OutConnectedPolygonIDs) const;
 
 	/** Returns the number of polygons connected to this edge */
 	UFUNCTION(BlueprintPure, Category="MeshDescription")
-	int32 GetNumEdgeConnectedPolygons(FEdgeID EdgeID) const;
+	MESHDESCRIPTION_API int32 GetNumEdgeConnectedPolygons(FEdgeID EdgeID) const;
 
 	/** Returns the vertex ID corresponding to one of the edge endpoints */
 	UFUNCTION(BlueprintPure, Category="MeshDescription")
-	FVertexID GetEdgeVertex(FEdgeID EdgeID, int32 VertexNumber) const;
+	MESHDESCRIPTION_API FVertexID GetEdgeVertex(FEdgeID EdgeID, int32 VertexNumber) const;
 
 	/** Returns a pair of vertex IDs defining the edge */
 	UFUNCTION(BlueprintPure, Category="MeshDescription")
-	void GetEdgeVertices(const FEdgeID EdgeID, TArray<FVertexID>& OutVertexIDs) const;
+	MESHDESCRIPTION_API void GetEdgeVertices(const FEdgeID EdgeID, TArray<FVertexID>& OutVertexIDs) const;
 
 
 	//////////////////////////////////////////////////////////////////////
@@ -383,39 +383,39 @@ public:
 
 	/** Get the polygon which contains this triangle */
 	UFUNCTION(BlueprintPure, Category="MeshDescription")
-	FPolygonID GetTrianglePolygon(FTriangleID TriangleID) const;
+	MESHDESCRIPTION_API FPolygonID GetTrianglePolygon(FTriangleID TriangleID) const;
 
 	/** Get the polygon group which contains this triangle */
 	UFUNCTION(BlueprintPure, Category="MeshDescription")
-	FPolygonGroupID GetTrianglePolygonGroup(FTriangleID TriangleID) const;
+	MESHDESCRIPTION_API FPolygonGroupID GetTrianglePolygonGroup(FTriangleID TriangleID) const;
 
 	/** Determines if this triangle is part of an n-gon */
 	UFUNCTION(BlueprintPure, Category="MeshDescription")
-	bool IsTrianglePartOfNgon(FTriangleID TriangleID) const;
+	MESHDESCRIPTION_API bool IsTrianglePartOfNgon(FTriangleID TriangleID) const;
 
 	/** Get the vertex instances which define this triangle */
 	UFUNCTION(BlueprintPure, Category="MeshDescription")
-	void GetTriangleVertexInstances(FTriangleID TriangleID, TArray<FVertexInstanceID>& OutVertexInstanceIDs) const;
+	MESHDESCRIPTION_API void GetTriangleVertexInstances(FTriangleID TriangleID, TArray<FVertexInstanceID>& OutVertexInstanceIDs) const;
 
 	/** Get the specified vertex instance by index */
 	UFUNCTION(BlueprintPure, Category="MeshDescription")
-	FVertexInstanceID GetTriangleVertexInstance(FTriangleID TriangleID, int32 Index) const;
+	MESHDESCRIPTION_API FVertexInstanceID GetTriangleVertexInstance(FTriangleID TriangleID, int32 Index) const;
 
 	/** Returns the vertices which define this triangle */
 	UFUNCTION(BlueprintPure, Category="MeshDescription")
-	void GetTriangleVertices(FTriangleID TriangleID, TArray<FVertexID>& OutVertexIDs) const;
+	MESHDESCRIPTION_API void GetTriangleVertices(FTriangleID TriangleID, TArray<FVertexID>& OutVertexIDs) const;
 
 	/** Returns the edges which define this triangle */
 	UFUNCTION(BlueprintPure, Category="MeshDescription")
-	void GetTriangleEdges(FTriangleID TriangleID, TArray<FEdgeID>& OutEdgeIDs) const;
+	MESHDESCRIPTION_API void GetTriangleEdges(FTriangleID TriangleID, TArray<FEdgeID>& OutEdgeIDs) const;
 
 	/** Returns the adjacent triangles to this triangle */
 	UFUNCTION(BlueprintPure, Category="MeshDescription")
-	void GetTriangleAdjacentTriangles(FTriangleID TriangleID, TArray<FTriangleID>& OutTriangleIDs) const;
+	MESHDESCRIPTION_API void GetTriangleAdjacentTriangles(FTriangleID TriangleID, TArray<FTriangleID>& OutTriangleIDs) const;
 
 	/** Return the vertex instance which corresponds to the given vertex on the given triangle, or INDEX_NONE */
 	UFUNCTION(BlueprintPure, Category="MeshDescription")
-	FVertexInstanceID GetVertexInstanceForTriangleVertex(FTriangleID TriangleID, FVertexID VertexID) const;
+	MESHDESCRIPTION_API FVertexInstanceID GetVertexInstanceForTriangleVertex(FTriangleID TriangleID, FVertexID VertexID) const;
 
 
 	//////////////////////////////////////////////////////////////////////
@@ -423,64 +423,64 @@ public:
 
 	/** Return reference to an array of triangle IDs which comprise this polygon */
 	UFUNCTION(BlueprintPure, Category="MeshDescription")
-	void GetPolygonTriangles(FPolygonID PolygonID, TArray<FTriangleID>& OutTriangleIDs) const;
+	MESHDESCRIPTION_API void GetPolygonTriangles(FPolygonID PolygonID, TArray<FTriangleID>& OutTriangleIDs) const;
 
 	/** Return the number of triangles which comprise this polygon */
 	UFUNCTION(BlueprintPure, Category="MeshDescription")
-	int32 GetNumPolygonTriangles(FPolygonID PolygonID) const;
+	MESHDESCRIPTION_API int32 GetNumPolygonTriangles(FPolygonID PolygonID) const;
 
 	/** Returns reference to an array of VertexInstance IDs forming the perimeter of this polygon */
 	UFUNCTION(BlueprintPure, Category="MeshDescription")
-	void GetPolygonVertexInstances(FPolygonID PolygonID, TArray<FVertexInstanceID>& OutVertexInstanceIDs) const;
+	MESHDESCRIPTION_API void GetPolygonVertexInstances(FPolygonID PolygonID, TArray<FVertexInstanceID>& OutVertexInstanceIDs) const;
 
 	/** Returns the number of vertices this polygon has */
 	UFUNCTION(BlueprintPure, Category="MeshDescription")
-	int32 GetNumPolygonVertices(FPolygonID PolygonID) const;
+	MESHDESCRIPTION_API int32 GetNumPolygonVertices(FPolygonID PolygonID) const;
 
 	/** Returns the vertices which form the polygon perimeter */
 	UFUNCTION(BlueprintPure, Category="MeshDescription")
-	void GetPolygonVertices(FPolygonID PolygonID, TArray<FVertexID>& OutVertexIDs) const;
+	MESHDESCRIPTION_API void GetPolygonVertices(FPolygonID PolygonID, TArray<FVertexID>& OutVertexIDs) const;
 
 	/** Returns the edges which form the polygon perimeter */
 	UFUNCTION(BlueprintPure, Category="MeshDescription")
-	void GetPolygonPerimeterEdges(FPolygonID PolygonID, TArray<FEdgeID>& OutEdgeIDs) const;
+	MESHDESCRIPTION_API void GetPolygonPerimeterEdges(FPolygonID PolygonID, TArray<FEdgeID>& OutEdgeIDs) const;
 
 	/** Populate the provided array with a list of edges which are internal to the polygon, i.e. those which separate
 	    constituent triangles. */
 	UFUNCTION(BlueprintPure, Category="MeshDescription")
-	void GetPolygonInternalEdges(FPolygonID PolygonID, TArray<FEdgeID>& OutEdgeIDs) const;
+	MESHDESCRIPTION_API void GetPolygonInternalEdges(FPolygonID PolygonID, TArray<FEdgeID>& OutEdgeIDs) const;
 
 	/** Return the number of internal edges in this polygon */
 	UFUNCTION(BlueprintPure, Category="MeshDescription")
-	int32 GetNumPolygonInternalEdges(FPolygonID PolygonID) const;
+	MESHDESCRIPTION_API int32 GetNumPolygonInternalEdges(FPolygonID PolygonID) const;
 
 	/** Populates the passed array with adjacent polygons */
 	UFUNCTION(BlueprintPure, Category="MeshDescription")
-	void GetPolygonAdjacentPolygons(FPolygonID PolygonID, TArray<FPolygonID>& OutPolygonIDs) const;
+	MESHDESCRIPTION_API void GetPolygonAdjacentPolygons(FPolygonID PolygonID, TArray<FPolygonID>& OutPolygonIDs) const;
 
 	/** Return the polygon group associated with a polygon */
 	UFUNCTION(BlueprintPure, Category="MeshDescription")
-	FPolygonGroupID GetPolygonPolygonGroup(FPolygonID PolygonID) const;
+	MESHDESCRIPTION_API FPolygonGroupID GetPolygonPolygonGroup(FPolygonID PolygonID) const;
 
 	/** Return the vertex instance which corresponds to the given vertex on the given polygon, or INDEX_NONE */
 	UFUNCTION(BlueprintPure, Category="MeshDescription")
-	FVertexInstanceID GetVertexInstanceForPolygonVertex(FPolygonID PolygonID, FVertexID VertexID) const;
+	MESHDESCRIPTION_API FVertexInstanceID GetVertexInstanceForPolygonVertex(FPolygonID PolygonID, FVertexID VertexID) const;
 
 	/** Set the vertex instance at the given index around the polygon to the new value */
 	UFUNCTION(BlueprintCallable, Category="MeshDescription")
-	void SetPolygonVertexInstances(FPolygonID PolygonID, const TArray<FVertexInstanceID>& VertexInstanceIDs);
+	MESHDESCRIPTION_API void SetPolygonVertexInstances(FPolygonID PolygonID, const TArray<FVertexInstanceID>& VertexInstanceIDs);
 
 	/** Sets the polygon group associated with a polygon */
 	UFUNCTION(BlueprintCallable, Category="MeshDescription")
-	void SetPolygonPolygonGroup(FPolygonID PolygonID, FPolygonGroupID PolygonGroupID);
+	MESHDESCRIPTION_API void SetPolygonPolygonGroup(FPolygonID PolygonID, FPolygonGroupID PolygonGroupID);
 
 	/** Reverse the winding order of the vertices of this polygon */
 	UFUNCTION(BlueprintCallable, Category="MeshDescription")
-	void ReversePolygonFacing(FPolygonID PolygonID);
+	MESHDESCRIPTION_API void ReversePolygonFacing(FPolygonID PolygonID);
 
 	/** Generates triangles and internal edges for the given polygon */
 	UFUNCTION(BlueprintCallable, Category="MeshDescription")
-	void ComputePolygonTriangulation(FPolygonID PolygonID);
+	MESHDESCRIPTION_API void ComputePolygonTriangulation(FPolygonID PolygonID);
 
 
 	//////////////////////////////////////////////////////////////////////
@@ -488,11 +488,11 @@ public:
 
 	/** Returns the polygons associated with the given polygon group */
 	UFUNCTION(BlueprintPure, Category="MeshDescription")
-	void GetPolygonGroupPolygons(FPolygonGroupID PolygonGroupID, TArray<FPolygonID>& OutPolygonIDs) const;
+	MESHDESCRIPTION_API void GetPolygonGroupPolygons(FPolygonGroupID PolygonGroupID, TArray<FPolygonID>& OutPolygonIDs) const;
 
 	/** Returns the number of polygons in this polygon group */
 	UFUNCTION(BlueprintPure, Category="MeshDescription")
-	int32 GetNumPolygonGroupPolygons(FPolygonGroupID PolygonGroupID) const;
+	MESHDESCRIPTION_API int32 GetNumPolygonGroupPolygons(FPolygonGroupID PolygonGroupID) const;
 
 protected:
 	FMeshDescription OwnedMeshDescription;
diff --git a/Engine/Source/Runtime/MeshDescription/Public/MeshDescriptionBaseBulkData.h b/Engine/Source/Runtime/MeshDescription/Public/MeshDescriptionBaseBulkData.h
--- a/Engine/Source/Runtime/MeshDescription/Public/MeshDescriptionBaseBulkData.h
+++ b/Engine/Source/Runtime/MeshDescription/Public/MeshDescriptionBaseBulkData.h
@@ -20,32 +20,32 @@ struct FMeshDescriptionBulkData;
 /**
  * UObject wrapper for FMeshDescriptionBulkData
  */
-UCLASS()
-class MESHDESCRIPTION_API UMeshDescriptionBaseBulkData : public UObject
+UCLASS(MinimalAPI)
+class UMeshDescriptionBaseBulkData : public UObject
 {
 	GENERATED_BODY()
 
 public:
-	UMeshDescriptionBaseBulkData();
-	virtual void Serialize(FArchive& Ar) override;
-	virtual bool IsEditorOnly() const override;
-	virtual bool NeedsLoadForClient() const override;
-	virtual bool NeedsLoadForServer() const override;
-	virtual bool NeedsLoadForEditorGame() const override;
+	MESHDESCRIPTION_API UMeshDescriptionBaseBulkData();
+	MESHDESCRIPTION_API virtual void Serialize(FArchive& Ar) override;
+	MESHDESCRIPTION_API virtual bool IsEditorOnly() const override;
+	MESHDESCRIPTION_API virtual bool NeedsLoadForClient() const override;
+	MESHDESCRIPTION_API virtual bool NeedsLoadForServer() const override;
+	MESHDESCRIPTION_API virtual bool NeedsLoadForEditorGame() const override;
 
 #if WITH_EDITORONLY_DATA
 
-	void Empty();
-	UMeshDescriptionBase* CreateMeshDescription();
-	UMeshDescriptionBase* GetMeshDescription() const;
-	bool HasCachedMeshDescription() const;
-	bool CacheMeshDescription();
-	void CommitMeshDescription(bool bUseHashAsGuid);
-	void RemoveMeshDescription();
-	bool IsBulkDataValid() const;
-
-	const FMeshDescriptionBulkData& GetBulkData() const;
-	FMeshDescriptionBulkData& GetBulkData();
+	MESHDESCRIPTION_API void Empty();
+	MESHDESCRIPTION_API UMeshDescriptionBase* CreateMeshDescription();
+	MESHDESCRIPTION_API UMeshDescriptionBase* GetMeshDescription() const;
+	MESHDESCRIPTION_API bool HasCachedMeshDescription() const;
+	MESHDESCRIPTION_API bool CacheMeshDescription();
+	MESHDESCRIPTION_API void CommitMeshDescription(bool bUseHashAsGuid);
+	MESHDESCRIPTION_API void RemoveMeshDescription();
+	MESHDESCRIPTION_API bool IsBulkDataValid() const;
+
+	MESHDESCRIPTION_API const FMeshDescriptionBulkData& GetBulkData() const;
+	MESHDESCRIPTION_API FMeshDescriptionBulkData& GetBulkData();
 
 protected:
 	TPimplPtr<FMeshDescriptionBulkData> BulkData;
diff --git a/Engine/Source/Runtime/MeshDescription/Public/MeshElementIndexer.h b/Engine/Source/Runtime/MeshDescription/Public/MeshElementIndexer.h
--- a/Engine/Source/Runtime/MeshDescription/Public/MeshElementIndexer.h
+++ b/Engine/Source/Runtime/MeshDescription/Public/MeshElementIndexer.h
@@ -44,11 +44,11 @@ struct FElementID;
  * there are generally few polygon groups compared to many triangles. Chunking can be turned off for such indexers,
  * resulting in a more simple array representation which can be arbitrarily grown without requiring elements to be inserted.
  */
-class MESHDESCRIPTION_API FMeshElementIndexer
+class FMeshElementIndexer
 {
 public:
 	FMeshElementIndexer() = default;
-	FMeshElementIndexer(const FMeshElementChannels* Key, const FMeshElementChannels* Referencer, FName AttributeName, int32 ReferencerChannelIndex = 0);
+	MESHDESCRIPTION_API FMeshElementIndexer(const FMeshElementChannels* Key, const FMeshElementChannels* Referencer, FName AttributeName, int32 ReferencerChannelIndex = 0);
 
 	// Copying not allowed as they reference elements of a MeshDescription of which they are normally a part
 	FMeshElementIndexer(const FMeshElementIndexer&) = delete;
@@ -61,7 +61,7 @@ public:
 	/**
 	 * Set the parameters of this mesh element indexer
 	 */
-	void Set(const FMeshElementChannels* Key, const FMeshElementChannels* Referencer, FName AttributeName, int32 ReferencerChannelIndex = 0);
+	MESHDESCRIPTION_API void Set(const FMeshElementChannels* Key, const FMeshElementChannels* Referencer, FName AttributeName, int32 ReferencerChannelIndex = 0);
 
 	/**
 	 * Sets the initial/expected number of references per key
@@ -92,7 +92,7 @@ public:
 	/**
 	 * Resets the indexer
 	 */
-	void Reset();
+	MESHDESCRIPTION_API void Reset();
 
 	/**
 	 * Retrieve all referencer indices whose named attribute refers to KeyIndex.
@@ -135,20 +135,20 @@ public:
 	/**
 	 * Remove the specified key from the indexer
 	 */
-	void RemoveKey(int32 KeyIndex, int32 KeyChannelIndex = 0);
+	MESHDESCRIPTION_API void RemoveKey(int32 KeyIndex, int32 KeyChannelIndex = 0);
 
 	/**
 	 * Remove reference from key.
 	 * This will take immediate effect.
 	 */
-	void RemoveReferenceFromKey(int32 KeyIndex, int32 ReferenceIndex, int32 KeyChannelIndex = 0);
+	MESHDESCRIPTION_API void RemoveReferenceFromKey(int32 KeyIndex, int32 ReferenceIndex, int32 KeyChannelIndex = 0);
 
 	/**
 	 * Add a reference to a key.
 	 * This will be queued and take effect in batch when an affected key is looked up.
 	 * This is preferable to regenerating the index entirely.
 	 */
-	void AddReferenceToKey(int32 KeyIndex, int32 ReferenceIndex, int32 KeyChannelIndex = 0);
+	MESHDESCRIPTION_API void AddReferenceToKey(int32 KeyIndex, int32 ReferenceIndex, int32 KeyChannelIndex = 0);
 
 	/**
 	 * Suspend indexing until the next rebuild
@@ -167,28 +167,28 @@ public:
 	/**
 	 * Force the indexer to be completely rebuilt.
 	 */
-	void ForceRebuild();
+	MESHDESCRIPTION_API void ForceRebuild();
 
 	/**
 	 * Build any stale indices
 	 */
-	void Build();
+	MESHDESCRIPTION_API void Build();
 
 private:
 	// Ensure that the right number of per channel indices are allocated
-	void UpdatePerChannelIndices();
+	MESHDESCRIPTION_API void UpdatePerChannelIndices();
 
 	// Incrementally builds the index based on the lists of stale keys and referencers
-	void BuildIndex(int32 Index);
+	MESHDESCRIPTION_API void BuildIndex(int32 Index);
 
 	// Does a full rebuild of the index from scratch, directly from referencing attributes
-	void RebuildIndex(int32 Index);
+	MESHDESCRIPTION_API void RebuildIndex(int32 Index);
 
 	// Performs an incremental build if the specified key is stale
-	void ConditionalBuild(int32 KeyIndex, int32 KeyChannelIndex);
+	MESHDESCRIPTION_API void ConditionalBuild(int32 KeyIndex, int32 KeyChannelIndex);
 
 	// Add a reference to a key in unchunked mode
-	void AddUnchunked(int32 KeyIndex, int32 ReferenceIndex, int32 KeyChannelIndex = 0);
+	MESHDESCRIPTION_API void AddUnchunked(int32 KeyIndex, int32 ReferenceIndex, int32 KeyChannelIndex = 0);
 
 	// Return the chunk index relating to this key index
 	int32 GetChunkForKey(int32 ElementIndex) const
diff --git a/Engine/Source/Runtime/MeshUtilitiesCommon/Public/Allocator2D.h b/Engine/Source/Runtime/MeshUtilitiesCommon/Public/Allocator2D.h
--- a/Engine/Source/Runtime/MeshUtilitiesCommon/Public/Allocator2D.h
+++ b/Engine/Source/Runtime/MeshUtilitiesCommon/Public/Allocator2D.h
@@ -16,7 +16,7 @@ template <typename FuncType> class TFunctionRef;
 
 #define DEBUG_LAYOUT_STATS 0
 
-class MESHUTILITIESCOMMON_API FAllocator2D
+class FAllocator2D
 {
 public:
 	enum class EMode
@@ -61,58 +61,58 @@ public:
 	};
 
 public:
-	FAllocator2D( EMode Mode, uint32 Width, uint32 Height, ELightmapUVVersion LayoutVersion );
+	MESHUTILITIESCOMMON_API FAllocator2D( EMode Mode, uint32 Width, uint32 Height, ELightmapUVVersion LayoutVersion );
 
 	// Must clear before using
-	void       Clear();
+	MESHUTILITIESCOMMON_API void       Clear();
 
-	bool       Find( FRect& Rect );
-	bool       Test( FRect Rect );
-	void       Alloc( FRect Rect );
+	MESHUTILITIESCOMMON_API bool       Find( FRect& Rect );
+	MESHUTILITIESCOMMON_API bool       Test( FRect Rect );
+	MESHUTILITIESCOMMON_API void       Alloc( FRect Rect );
 
-	bool       FindBitByBit( FRect& Rect, const FAllocator2D& Other );
-	bool       FindWithSegments( FRect& Rect, const FAllocator2D& Other, TFunctionRef<bool (const FAllocator2D::FRect&)> IsBestRect ) const;
-	bool       Test( FRect Rect, const FAllocator2D& Other );
-	void       Alloc( FRect Rect, const FAllocator2D& Other );
+	MESHUTILITIESCOMMON_API bool       FindBitByBit( FRect& Rect, const FAllocator2D& Other );
+	MESHUTILITIESCOMMON_API bool       FindWithSegments( FRect& Rect, const FAllocator2D& Other, TFunctionRef<bool (const FAllocator2D::FRect&)> IsBestRect ) const;
+	MESHUTILITIESCOMMON_API bool       Test( FRect Rect, const FAllocator2D& Other );
+	MESHUTILITIESCOMMON_API void       Alloc( FRect Rect, const FAllocator2D& Other );
 
-	uint64     GetBit( uint32 x, uint32 y ) const;
-	void       SetBit( uint32 x, uint32 y );
-	void       ClearBit( uint32 x, uint32 y );
+	MESHUTILITIESCOMMON_API uint64     GetBit( uint32 x, uint32 y ) const;
+	MESHUTILITIESCOMMON_API void       SetBit( uint32 x, uint32 y );
+	MESHUTILITIESCOMMON_API void       ClearBit( uint32 x, uint32 y );
 
-	void       CreateUsedSegments();
-	void       MergeRun( FRun& Run, const FRun& OtherRun, uint32 RectOffset, uint32 RectLength, uint32 PrimaryResolution /* Resolution along the axis the run belongs to */, uint32 PerpendicularResolution );
-	void       MergeSegments( const FRect& Rect, const FAllocator2D& Other );
+	MESHUTILITIESCOMMON_API void       CreateUsedSegments();
+	MESHUTILITIESCOMMON_API void       MergeRun( FRun& Run, const FRun& OtherRun, uint32 RectOffset, uint32 RectLength, uint32 PrimaryResolution /* Resolution along the axis the run belongs to */, uint32 PerpendicularResolution );
+	MESHUTILITIESCOMMON_API void       MergeSegments( const FRect& Rect, const FAllocator2D& Other );
 
-	void       FlipX( const FRect& Rect );
-	void       FlipY( const FRect& Rect );
+	MESHUTILITIESCOMMON_API void       FlipX( const FRect& Rect );
+	MESHUTILITIESCOMMON_API void       FlipY( const FRect& Rect );
 
-	uint32     GetUsedTexels() const;
+	MESHUTILITIESCOMMON_API uint32     GetUsedTexels() const;
 
-	void       CopyRuns( TArray<FRun>& Runs, const TArray<FRun>& OtherRuns, int32 MaxSize );
+	MESHUTILITIESCOMMON_API void       CopyRuns( TArray<FRun>& Runs, const TArray<FRun>& OtherRuns, int32 MaxSize );
 
 	// Take control of the copy assignment to reduce the amount of data movement to the strict minimum
 	FAllocator2D(const FAllocator2D& Other) = default;
-	FAllocator2D& operator = (const FAllocator2D& Other);
+	MESHUTILITIESCOMMON_API FAllocator2D& operator = (const FAllocator2D& Other);
 
 	// Allow to visualize the content in ascii for debugging purpose. (i.e Watch or Immediate window).
-	FString    ToString() const;
+	MESHUTILITIESCOMMON_API FString    ToString() const;
 
 	// Get the MD5 hash of the rasterized content
-	FMD5Hash   GetRasterMD5() const;
+	MESHUTILITIESCOMMON_API FMD5Hash   GetRasterMD5() const;
 
 	uint32     GetRasterWidth()  const { return RasterWidth; }
 	uint32     GetRasterHeight() const { return RasterHeight; }
 
-	void       ResetStats();
-	void       PublishStats( int32 ChartIndex, int32 Orientation, bool bFound, const FRect& Rect, const FRect& BestRect, const FMD5Hash& ChartMD5, TFunctionRef<bool (const FAllocator2D::FRect&)> IsBestRect );
+	MESHUTILITIESCOMMON_API void       ResetStats();
+	MESHUTILITIESCOMMON_API void       PublishStats( int32 ChartIndex, int32 Orientation, bool bFound, const FRect& Rect, const FRect& BestRect, const FMD5Hash& ChartMD5, TFunctionRef<bool (const FAllocator2D::FRect&)> IsBestRect );
 
 protected:
-	bool       TestOneRun( const FRun& Run, const FRun& OtherRun, uint32 RectOffset, uint32 RectLength, uint32 PrimaryResolution, uint32& OutFailedLength ) const;
-	bool       TestAllRows( const FRect& Rect, const FAllocator2D& Other, uint32& FailedLength ) const;
-	bool       TestAllColumns( const FRect& Rect, const FAllocator2D& Other, uint32& FailedLength ) const;
-	void       InitRuns( TArray<FRun>& Runs, uint32 PrimaryResolution, uint32 PerpendicularRasterSize);
-	void       InitSegments();
-	void       AddUsedSegment( FRun& Run, uint32 StartPos, uint32 Length );
+	MESHUTILITIESCOMMON_API bool       TestOneRun( const FRun& Run, const FRun& OtherRun, uint32 RectOffset, uint32 RectLength, uint32 PrimaryResolution, uint32& OutFailedLength ) const;
+	MESHUTILITIESCOMMON_API bool       TestAllRows( const FRect& Rect, const FAllocator2D& Other, uint32& FailedLength ) const;
+	MESHUTILITIESCOMMON_API bool       TestAllColumns( const FRect& Rect, const FAllocator2D& Other, uint32& FailedLength ) const;
+	MESHUTILITIESCOMMON_API void       InitRuns( TArray<FRun>& Runs, uint32 PrimaryResolution, uint32 PerpendicularRasterSize);
+	MESHUTILITIESCOMMON_API void       InitSegments();
+	MESHUTILITIESCOMMON_API void       AddUsedSegment( FRun& Run, uint32 StartPos, uint32 Length );
 
 	// Enforce that those cannot be changed in flight
 	const EMode              Mode;
diff --git a/Engine/Source/Runtime/MeshUtilitiesCommon/Public/LayoutUV.h b/Engine/Source/Runtime/MeshUtilitiesCommon/Public/LayoutUV.h
--- a/Engine/Source/Runtime/MeshUtilitiesCommon/Public/LayoutUV.h
+++ b/Engine/Source/Runtime/MeshUtilitiesCommon/Public/LayoutUV.h
@@ -37,7 +37,7 @@ struct FMeshChart
 
 struct FOverlappingCorners;
 
-class MESHUTILITIESCOMMON_API FLayoutUV
+class FLayoutUV
 {
 public:
 
@@ -61,14 +61,14 @@ public:
 		virtual void        SetOutputTexcoord(uint32 Index, const FVector2f& Value) = 0;
 	};
 
-	FLayoutUV( IMeshView& InMeshView );
+	MESHUTILITIESCOMMON_API FLayoutUV( IMeshView& InMeshView );
 	void SetVersion( ELightmapUVVersion Version ) { LayoutVersion = Version; }
-	int32 FindCharts( const FOverlappingCorners& OverlappingCorners );
-	bool FindBestPacking( uint32 InTextureResolution );
-	void CommitPackedUVs();
+	MESHUTILITIESCOMMON_API int32 FindCharts( const FOverlappingCorners& OverlappingCorners );
+	MESHUTILITIESCOMMON_API bool FindBestPacking( uint32 InTextureResolution );
+	MESHUTILITIESCOMMON_API void CommitPackedUVs();
 
-	static void LogStats();
-	static void ResetStats();
+	static MESHUTILITIESCOMMON_API void LogStats();
+	static MESHUTILITIESCOMMON_API void ResetStats();
 private:
 	IMeshView& MeshView;
 	ELightmapUVVersion LayoutVersion;
@@ -81,7 +81,7 @@ private:
 	struct FChartFinder;
 	struct FChartPacker;
 
-	static TAtomic<uint64> FindBestPackingCount;
-	static TAtomic<uint64> FindBestPackingCycles;
-	static TAtomic<uint64> FindBestPackingEfficiency;
+	static MESHUTILITIESCOMMON_API TAtomic<uint64> FindBestPackingCount;
+	static MESHUTILITIESCOMMON_API TAtomic<uint64> FindBestPackingCycles;
+	static MESHUTILITIESCOMMON_API TAtomic<uint64> FindBestPackingEfficiency;
 };
diff --git a/Engine/Source/Runtime/MeshUtilitiesCommon/Public/OverlappingCorners.h b/Engine/Source/Runtime/MeshUtilitiesCommon/Public/OverlappingCorners.h
--- a/Engine/Source/Runtime/MeshUtilitiesCommon/Public/OverlappingCorners.h
+++ b/Engine/Source/Runtime/MeshUtilitiesCommon/Public/OverlappingCorners.h
@@ -12,24 +12,24 @@
 /**
 * Container to hold overlapping corners. For a vertex, lists all the overlapping vertices
 */
-struct MESHUTILITIESCOMMON_API FOverlappingCorners
+struct FOverlappingCorners
 {
 	FOverlappingCorners() {}
 
-	FOverlappingCorners(const TArray<FVector3f>& InVertices, const TArray<uint32>& InIndices, float ComparisonThreshold);
-	FOverlappingCorners(const FLayoutUV::IMeshView& MeshView, float ComparisonThreshold);
+	MESHUTILITIESCOMMON_API FOverlappingCorners(const TArray<FVector3f>& InVertices, const TArray<uint32>& InIndices, float ComparisonThreshold);
+	MESHUTILITIESCOMMON_API FOverlappingCorners(const FLayoutUV::IMeshView& MeshView, float ComparisonThreshold);
 
 	/* Resets, pre-allocates memory, marks all indices as not overlapping in preperation for calls to Add() */
-	void Init(int32 NumIndices);
+	MESHUTILITIESCOMMON_API void Init(int32 NumIndices);
 
 	/* Add overlapping indices pair */
-	void Add(int32 Key, int32 Value);
+	MESHUTILITIESCOMMON_API void Add(int32 Key, int32 Value);
 
 	/* Sorts arrays, converts sets to arrays for sorting and to allow simple iterating code, prevents additional adding */
-	void FinishAdding();
+	MESHUTILITIESCOMMON_API void FinishAdding();
 
 	/* Estimate memory allocated */
-	uint32 GetAllocatedSize(void) const;
+	MESHUTILITIESCOMMON_API uint32 GetAllocatedSize(void) const;
 
 	/**
 	* @return array of sorted overlapping indices including input 'Key', empty array for indices that have no overlaps.
diff --git a/Engine/Source/Runtime/MessagingRpc/Public/MessageRpcServer.h b/Engine/Source/Runtime/MessagingRpc/Public/MessageRpcServer.h
--- a/Engine/Source/Runtime/MessagingRpc/Public/MessageRpcServer.h
+++ b/Engine/Source/Runtime/MessagingRpc/Public/MessageRpcServer.h
@@ -26,30 +26,30 @@ struct FMessageRpcCancel;
 /**
  * Implements an RPC server.
  */
-class MESSAGINGRPC_API FMessageRpcServer
+class FMessageRpcServer
 	: public IMessageRpcServer
 {
 public:
 	/** Default constructor. */
-	FMessageRpcServer();
-	FMessageRpcServer(const FString& InDebugName, const TSharedRef<IMessageBus, ESPMode::ThreadSafe>& InMessageBus);
+	MESSAGINGRPC_API FMessageRpcServer();
+	MESSAGINGRPC_API FMessageRpcServer(const FString& InDebugName, const TSharedRef<IMessageBus, ESPMode::ThreadSafe>& InMessageBus);
 
 	/** Virtual destructor. */
-	virtual ~FMessageRpcServer();
+	MESSAGINGRPC_API virtual ~FMessageRpcServer();
 
 public:
 
 	//~ IMessageRpcServer interface
 
-	virtual void AddHandler(const FTopLevelAssetPath& RequestMessageType, const TSharedRef<IMessageRpcHandler>& Handler) override;
-	virtual const FMessageAddress& GetAddress() const override;
-PRAGMA_DISABLE_DEPRECATION_WARNINGS
+	MESSAGINGRPC_API virtual void AddHandler(const FTopLevelAssetPath& RequestMessageType, const TSharedRef<IMessageRpcHandler>& Handler) override;
+	MESSAGINGRPC_API virtual const FMessageAddress& GetAddress() const override;
+MESSAGINGRPC_API PRAGMA_DISABLE_DEPRECATION_WARNINGS
 	virtual FOnMessageRpcNoHandler& OnNoHandler() override;
-PRAGMA_ENABLE_DEPRECATION_WARNINGS
+MESSAGINGRPC_API PRAGMA_ENABLE_DEPRECATION_WARNINGS
 	virtual FOnMessagePathNameRpcNoHandler& OnNoHandlerWithPathName() override;
-	virtual void SetSendProgressUpdate(bool InSendProgress) override;
+	MESSAGINGRPC_API virtual void SetSendProgressUpdate(bool InSendProgress) override;
 protected:
-	explicit FMessageRpcServer(FMessageEndpointBuilder&& InEndpointBuilder);
+	MESSAGINGRPC_API explicit FMessageRpcServer(FMessageEndpointBuilder&& InEndpointBuilder);
 
 	struct FReturnInfo
 	{
@@ -60,19 +60,19 @@ protected:
 		TSharedPtr<IAsyncTask> Task;
 	};
 
-	TSharedPtr<FMessageEndpoint, ESPMode::ThreadSafe> GetEndpoint() const;
+	MESSAGINGRPC_API TSharedPtr<FMessageEndpoint, ESPMode::ThreadSafe> GetEndpoint() const;
 
 	/** Processes an FMessageRpcCancel message. */
-	void ProcessCancelation(const FMessageRpcCancel& Message, const TSharedRef<IMessageContext, ESPMode::ThreadSafe>& Context);
+	MESSAGINGRPC_API void ProcessCancelation(const FMessageRpcCancel& Message, const TSharedRef<IMessageContext, ESPMode::ThreadSafe>& Context);
 
 	/** Processes an RPC request message. */
-	void ProcessRequest(const TSharedRef<IMessageContext, ESPMode::ThreadSafe>& Context);
+	MESSAGINGRPC_API void ProcessRequest(const TSharedRef<IMessageContext, ESPMode::ThreadSafe>& Context);
 
 	/** Send a progress message to the RPC client that made the RPC call. */
-	void SendProgress(const FGuid& CallId, const FReturnInfo& ReturnInfo);
+	MESSAGINGRPC_API void SendProgress(const FGuid& CallId, const FReturnInfo& ReturnInfo);
 
 	/** Send a result message to the RPC client that made the RPC call. */
-	void SendResult(const FGuid& CallId, const FReturnInfo& ReturnInfo);
+	MESSAGINGRPC_API void SendResult(const FGuid& CallId, const FReturnInfo& ReturnInfo);
 
 protected:
 
diff --git a/Engine/Source/Runtime/MoviePlayer/Public/MoviePlayer.h b/Engine/Source/Runtime/MoviePlayer/Public/MoviePlayer.h
--- a/Engine/Source/Runtime/MoviePlayer/Public/MoviePlayer.h
+++ b/Engine/Source/Runtime/MoviePlayer/Public/MoviePlayer.h
@@ -129,7 +129,7 @@ public:
 
 
 /** Struct of all the attributes a loading screen will have. */
-struct MOVIEPLAYER_API FLoadingScreenAttributes
+struct FLoadingScreenAttributes
 {
 	FLoadingScreenAttributes()
 		: MinimumLoadingScreenDisplayTime(-1.0f)
@@ -168,10 +168,10 @@ struct MOVIEPLAYER_API FLoadingScreenAttributes
 	TEnumAsByte<EMoviePlaybackType> PlaybackType;
 
 	/** True if there is either a standalone widget or any movie paths or both. */
-	bool IsValid() const;
+	MOVIEPLAYER_API bool IsValid() const;
 
 	/** Creates a simple test loading screen widget. */
-	static TSharedRef<class SWidget> NewTestLoadingScreenWidget();
+	static MOVIEPLAYER_API TSharedRef<class SWidget> NewTestLoadingScreenWidget();
 };
 
 
diff --git a/Engine/Source/Runtime/MoviePlayer/Public/MoviePlayerSettings.h b/Engine/Source/Runtime/MoviePlayer/Public/MoviePlayerSettings.h
--- a/Engine/Source/Runtime/MoviePlayer/Public/MoviePlayerSettings.h
+++ b/Engine/Source/Runtime/MoviePlayer/Public/MoviePlayerSettings.h
@@ -15,8 +15,8 @@
 /**
  * Implements the settings for the Windows target platform.
  */
-UCLASS(config=Game, defaultconfig)
-class MOVIEPLAYER_API UMoviePlayerSettings
+UCLASS(config=Game, defaultconfig, MinimalAPI)
+class UMoviePlayerSettings
 	: public UObject
 {
 public:
diff --git a/Engine/Source/Runtime/MoviePlayerProxy/Public/MoviePlayerProxy.h b/Engine/Source/Runtime/MoviePlayerProxy/Public/MoviePlayerProxy.h
--- a/Engine/Source/Runtime/MoviePlayerProxy/Public/MoviePlayerProxy.h
+++ b/Engine/Source/Runtime/MoviePlayerProxy/Public/MoviePlayerProxy.h
@@ -14,23 +14,23 @@ class IMoviePlayerProxyServer;
  * The movie player can call RegisterServer/UnregisterServer
  * so it can receive the calls from the engine.
  */ 
-class MOVIEPLAYERPROXY_API FMoviePlayerProxy
+class FMoviePlayerProxy
 {
 public:
 	/** Call this before doing a blocking operation on the game thread so that the movie player can activate. */
-	static void BlockingStarted();
+	static MOVIEPLAYERPROXY_API void BlockingStarted();
 	/** Call this periodically during a blocking operation on the game thread. */
-	static void BlockingTick();
+	static MOVIEPLAYERPROXY_API void BlockingTick();
 	/** Call this once the blocking operation is done to shut down the movie player. */
-	static void BlockingFinished();
+	static MOVIEPLAYERPROXY_API void BlockingFinished();
 	/** Call this to make sure the movie player is no longer running. */
-	static void BlockingForceFinished();
+	static MOVIEPLAYERPROXY_API void BlockingForceFinished();
 	/** Call this to prevent the movie player from using the Slate thread. */
-	static void SetIsSlateThreadAllowed(bool bInIsSlateThreadAllowed);
+	static MOVIEPLAYERPROXY_API void SetIsSlateThreadAllowed(bool bInIsSlateThreadAllowed);
 	/** Call this to hook up a server. */
-	static void RegisterServer(IMoviePlayerProxyServer* InServer);
+	static MOVIEPLAYERPROXY_API void RegisterServer(IMoviePlayerProxyServer* InServer);
 	/** Call this to unregister the current server. */
-	static void UnregisterServer();
+	static MOVIEPLAYERPROXY_API void UnregisterServer();
 	
 private:
 	/** Our current worker that handles blocks. */
@@ -48,7 +48,7 @@ private:
  * 
  * You can also manually call Finish if you want to trigger this earlier.
  */
-class MOVIEPLAYERPROXY_API FMoviePlayerProxyBlock
+class FMoviePlayerProxyBlock
 {
 public:
 	FMoviePlayerProxyBlock()
diff --git a/Engine/Source/Runtime/MoviePlayerProxy/Public/MoviePlayerProxyServer.h b/Engine/Source/Runtime/MoviePlayerProxy/Public/MoviePlayerProxyServer.h
--- a/Engine/Source/Runtime/MoviePlayerProxy/Public/MoviePlayerProxyServer.h
+++ b/Engine/Source/Runtime/MoviePlayerProxy/Public/MoviePlayerProxyServer.h
@@ -8,7 +8,7 @@
  * Implement this interface if you want to handle what happens when the engine interacts
  * with the MoviePlayer, etc when the game thread is blocked.
  */
-class MOVIEPLAYERPROXY_API IMoviePlayerProxyServer
+class IMoviePlayerProxyServer
 {
 public:
 	/** Called before doing a blocking operation on the game thread occurs so that the movie player can activate. */
diff --git a/Engine/Source/Runtime/MovieScene/Private/Channels/MovieSceneCurveChannelImpl.h b/Engine/Source/Runtime/MovieScene/Private/Channels/MovieSceneCurveChannelImpl.h
--- a/Engine/Source/Runtime/MovieScene/Private/Channels/MovieSceneCurveChannelImpl.h
+++ b/Engine/Source/Runtime/MovieScene/Private/Channels/MovieSceneCurveChannelImpl.h
@@ -21,7 +21,7 @@ namespace UE::MovieScene::Interpolation
 
 /** Utility class for curve channels */
 template<typename ChannelType>
-struct MOVIESCENE_API TMovieSceneCurveChannelImpl
+struct TMovieSceneCurveChannelImpl
 {
 	/** The type of channel value structs */
 	using ChannelValueType = typename ChannelType::ChannelValueType;
diff --git a/Engine/Source/Runtime/MovieScene/Public/ActorForWorldTransforms.h b/Engine/Source/Runtime/MovieScene/Public/ActorForWorldTransforms.h
--- a/Engine/Source/Runtime/MovieScene/Public/ActorForWorldTransforms.h
+++ b/Engine/Source/Runtime/MovieScene/Public/ActorForWorldTransforms.h
@@ -8,7 +8,7 @@
 
 //Description of an actor selected parts we can find world transforms on 
 USTRUCT(BlueprintType)
-struct  MOVIESCENE_API FActorForWorldTransforms
+struct  FActorForWorldTransforms
 {
 	GENERATED_BODY()
 
@@ -44,4 +44,4 @@ struct  MOVIESCENE_API FActorForWorldTransforms
 		}
 		return false;
 	}
-};
\ No newline at end of file
+};
diff --git a/Engine/Source/Runtime/MovieScene/Public/Channels/IMovieSceneChannelOverrideProvider.h b/Engine/Source/Runtime/MovieScene/Public/Channels/IMovieSceneChannelOverrideProvider.h
--- a/Engine/Source/Runtime/MovieScene/Public/Channels/IMovieSceneChannelOverrideProvider.h
+++ b/Engine/Source/Runtime/MovieScene/Public/Channels/IMovieSceneChannelOverrideProvider.h
@@ -135,8 +135,8 @@ struct TNamedChannelOverrideProviderTraits : FChannelOverrideProviderTraits
 }  // namespace MovieScene
 }  // namespace UE
 
-UINTERFACE()
-class MOVIESCENE_API UMovieSceneChannelOverrideProvider : public UInterface
+UINTERFACE(MinimalAPI)
+class UMovieSceneChannelOverrideProvider : public UInterface
 {
 public:
 
@@ -146,7 +146,7 @@ public:
 /**
  * Interface to be added to UMovieSceneSection types when they contain entity data
  */
-class MOVIESCENE_API IMovieSceneChannelOverrideProvider
+class IMovieSceneChannelOverrideProvider
 {
 public:
 
diff --git a/Engine/Source/Runtime/MovieScene/Public/Channels/MovieSceneAudioTriggerChannel.h b/Engine/Source/Runtime/MovieScene/Public/Channels/MovieSceneAudioTriggerChannel.h
--- a/Engine/Source/Runtime/MovieScene/Public/Channels/MovieSceneAudioTriggerChannel.h
+++ b/Engine/Source/Runtime/MovieScene/Public/Channels/MovieSceneAudioTriggerChannel.h
@@ -29,7 +29,7 @@ struct FMoveSceneAudioTriggerState
 };
 
 USTRUCT()
-struct MOVIESCENE_API FMovieSceneAudioTriggerChannel : public FMovieSceneChannel
+struct FMovieSceneAudioTriggerChannel : public FMovieSceneChannel
 {
 	GENERATED_BODY()
 
@@ -79,23 +79,23 @@ struct MOVIESCENE_API FMovieSceneAudioTriggerChannel : public FMovieSceneChannel
 	 * @param OutTriggered   A Value that if set to true combined with the function returning true means there's a trigger
 	 * @return true if the channel was evaluated successfully, false otherwise
 	 */
-	bool EvaluatePossibleTriggers(const FMovieSceneContext& InContext, FMoveSceneAudioTriggerState& InState, bool& OutTriggered) const;
+	MOVIESCENE_API bool EvaluatePossibleTriggers(const FMovieSceneContext& InContext, FMoveSceneAudioTriggerState& InState, bool& OutTriggered) const;
 
 public:
 
 	// ~ FMovieSceneChannel Interface
-	virtual void GetKeys(const TRange<FFrameNumber>& WithinRange, TArray<FFrameNumber>* OutKeyTimes, TArray<FKeyHandle>* OutKeyHandles) override;
-	virtual void GetKeyTimes(TArrayView<const FKeyHandle> InHandles, TArrayView<FFrameNumber> OutKeyTimes) override;
-	virtual void SetKeyTimes(TArrayView<const FKeyHandle> InHandles, TArrayView<const FFrameNumber> InKeyTimes) override;
-	virtual void DuplicateKeys(TArrayView<const FKeyHandle> InHandles, TArrayView<FKeyHandle> OutNewHandles) override;
-	virtual void DeleteKeys(TArrayView<const FKeyHandle> InHandles) override;
-	virtual void DeleteKeysFrom(FFrameNumber InTime, bool bDeleteKeysBefore) override;
-	virtual void ChangeFrameResolution(FFrameRate SourceRate, FFrameRate DestinationRate) override;
-	virtual TRange<FFrameNumber> ComputeEffectiveRange() const override;
-	virtual int32 GetNumKeys() const override;
-	virtual void Reset() override;
-	virtual void Offset(FFrameNumber DeltaPosition) override;
-	virtual void Optimize(const FKeyDataOptimizationParams& InParameters) override;
+	MOVIESCENE_API virtual void GetKeys(const TRange<FFrameNumber>& WithinRange, TArray<FFrameNumber>* OutKeyTimes, TArray<FKeyHandle>* OutKeyHandles) override;
+	MOVIESCENE_API virtual void GetKeyTimes(TArrayView<const FKeyHandle> InHandles, TArrayView<FFrameNumber> OutKeyTimes) override;
+	MOVIESCENE_API virtual void SetKeyTimes(TArrayView<const FKeyHandle> InHandles, TArrayView<const FFrameNumber> InKeyTimes) override;
+	MOVIESCENE_API virtual void DuplicateKeys(TArrayView<const FKeyHandle> InHandles, TArrayView<FKeyHandle> OutNewHandles) override;
+	MOVIESCENE_API virtual void DeleteKeys(TArrayView<const FKeyHandle> InHandles) override;
+	MOVIESCENE_API virtual void DeleteKeysFrom(FFrameNumber InTime, bool bDeleteKeysBefore) override;
+	MOVIESCENE_API virtual void ChangeFrameResolution(FFrameRate SourceRate, FFrameRate DestinationRate) override;
+	MOVIESCENE_API virtual TRange<FFrameNumber> ComputeEffectiveRange() const override;
+	MOVIESCENE_API virtual int32 GetNumKeys() const override;
+	MOVIESCENE_API virtual void Reset() override;
+	MOVIESCENE_API virtual void Offset(FFrameNumber DeltaPosition) override;
+	MOVIESCENE_API virtual void Optimize(const FKeyDataOptimizationParams& InParameters) override;
 
 private:
 
diff --git a/Engine/Source/Runtime/MovieScene/Public/Channels/MovieSceneBoolChannel.h b/Engine/Source/Runtime/MovieScene/Public/Channels/MovieSceneBoolChannel.h
--- a/Engine/Source/Runtime/MovieScene/Public/Channels/MovieSceneBoolChannel.h
+++ b/Engine/Source/Runtime/MovieScene/Public/Channels/MovieSceneBoolChannel.h
@@ -27,7 +27,7 @@ struct FPropertyTag;
 
 
 USTRUCT()
-struct MOVIESCENE_API FMovieSceneBoolChannel : public FMovieSceneChannel
+struct FMovieSceneBoolChannel : public FMovieSceneChannel
 {
 	GENERATED_BODY()
 
@@ -38,7 +38,7 @@ struct MOVIESCENE_API FMovieSceneBoolChannel : public FMovieSceneChannel
 	/**
 	 * Serialize this type from another
 	 */
-	bool SerializeFromMismatchedTag(const FPropertyTag& Tag, FStructuredArchive::FSlot Slot);
+	MOVIESCENE_API bool SerializeFromMismatchedTag(const FPropertyTag& Tag, FStructuredArchive::FSlot Slot);
 
 	/**
 	 * Access a mutable interface for this channel's data
@@ -91,7 +91,7 @@ struct MOVIESCENE_API FMovieSceneBoolChannel : public FMovieSceneChannel
 	 * @param OutValue   A value to receive the result
 	 * @return true if the channel was evaluated successfully, false otherwise
 	 */
-	virtual bool Evaluate(FFrameTime InTime, bool& OutValue) const;
+	MOVIESCENE_API virtual bool Evaluate(FFrameTime InTime, bool& OutValue) const;
 
 	/**
 	 * Add keys with these times to channel. The number of elements in both arrays much match or nothing is added.
@@ -108,19 +108,19 @@ struct MOVIESCENE_API FMovieSceneBoolChannel : public FMovieSceneChannel
 public:
 
 	// ~ FMovieSceneChannel Interface
-	virtual void GetKeys(const TRange<FFrameNumber>& WithinRange, TArray<FFrameNumber>* OutKeyTimes, TArray<FKeyHandle>* OutKeyHandles) override;
-	virtual void GetKeyTimes(TArrayView<const FKeyHandle> InHandles, TArrayView<FFrameNumber> OutKeyTimes) override;
-	virtual void SetKeyTimes(TArrayView<const FKeyHandle> InHandles, TArrayView<const FFrameNumber> InKeyTimes) override;
-	virtual void DuplicateKeys(TArrayView<const FKeyHandle> InHandles, TArrayView<FKeyHandle> OutNewHandles) override;
-	virtual void DeleteKeys(TArrayView<const FKeyHandle> InHandles) override;
-	virtual void DeleteKeysFrom(FFrameNumber InTime, bool bDeleteKeysBefore) override;
-	virtual void ChangeFrameResolution(FFrameRate SourceRate, FFrameRate DestinationRate) override;
-	virtual TRange<FFrameNumber> ComputeEffectiveRange() const override;
-	virtual int32 GetNumKeys() const override;
-	virtual void Reset() override;
-	virtual void Offset(FFrameNumber DeltaPosition) override;
-	virtual void Optimize(const FKeyDataOptimizationParams& InParameters) override;
-	virtual void ClearDefault() override;
+	MOVIESCENE_API virtual void GetKeys(const TRange<FFrameNumber>& WithinRange, TArray<FFrameNumber>* OutKeyTimes, TArray<FKeyHandle>* OutKeyHandles) override;
+	MOVIESCENE_API virtual void GetKeyTimes(TArrayView<const FKeyHandle> InHandles, TArrayView<FFrameNumber> OutKeyTimes) override;
+	MOVIESCENE_API virtual void SetKeyTimes(TArrayView<const FKeyHandle> InHandles, TArrayView<const FFrameNumber> InKeyTimes) override;
+	MOVIESCENE_API virtual void DuplicateKeys(TArrayView<const FKeyHandle> InHandles, TArrayView<FKeyHandle> OutNewHandles) override;
+	MOVIESCENE_API virtual void DeleteKeys(TArrayView<const FKeyHandle> InHandles) override;
+	MOVIESCENE_API virtual void DeleteKeysFrom(FFrameNumber InTime, bool bDeleteKeysBefore) override;
+	MOVIESCENE_API virtual void ChangeFrameResolution(FFrameRate SourceRate, FFrameRate DestinationRate) override;
+	MOVIESCENE_API virtual TRange<FFrameNumber> ComputeEffectiveRange() const override;
+	MOVIESCENE_API virtual int32 GetNumKeys() const override;
+	MOVIESCENE_API virtual void Reset() override;
+	MOVIESCENE_API virtual void Offset(FFrameNumber DeltaPosition) override;
+	MOVIESCENE_API virtual void Optimize(const FKeyDataOptimizationParams& InParameters) override;
+	MOVIESCENE_API virtual void ClearDefault() override;
 
 public:
 
diff --git a/Engine/Source/Runtime/MovieScene/Public/Channels/MovieSceneByteChannel.h b/Engine/Source/Runtime/MovieScene/Public/Channels/MovieSceneByteChannel.h
--- a/Engine/Source/Runtime/MovieScene/Public/Channels/MovieSceneByteChannel.h
+++ b/Engine/Source/Runtime/MovieScene/Public/Channels/MovieSceneByteChannel.h
@@ -28,7 +28,7 @@ struct FPropertyTag;
 
 
 USTRUCT()
-struct MOVIESCENE_API FMovieSceneByteChannel : public FMovieSceneChannel
+struct FMovieSceneByteChannel : public FMovieSceneChannel
 {
 	GENERATED_BODY()
 
@@ -42,7 +42,7 @@ struct MOVIESCENE_API FMovieSceneByteChannel : public FMovieSceneChannel
 	/**
 	 * Serialize this type from another
 	 */
-	bool SerializeFromMismatchedTag(const FPropertyTag& Tag, FStructuredArchive::FSlot Slot);
+	MOVIESCENE_API bool SerializeFromMismatchedTag(const FPropertyTag& Tag, FStructuredArchive::FSlot Slot);
 
 	/**
 	 * Access a mutable interface for this channel's data
@@ -95,7 +95,7 @@ struct MOVIESCENE_API FMovieSceneByteChannel : public FMovieSceneChannel
 	 * @param OutValue   A value to receive the result
 	 * @return true if the channel was evaluated successfully, false otherwise
 	 */
-	bool Evaluate(FFrameTime InTime, uint8& OutValue) const;
+	MOVIESCENE_API bool Evaluate(FFrameTime InTime, uint8& OutValue) const;
 
 	/**
 	 * Add keys with these times to channel. The number of elements in both arrays much match or nothing is added.
@@ -112,19 +112,19 @@ struct MOVIESCENE_API FMovieSceneByteChannel : public FMovieSceneChannel
 public:
 
 	// ~ FMovieSceneChannel Interface
-	virtual void GetKeys(const TRange<FFrameNumber>& WithinRange, TArray<FFrameNumber>* OutKeyTimes, TArray<FKeyHandle>* OutKeyHandles) override;
-	virtual void GetKeyTimes(TArrayView<const FKeyHandle> InHandles, TArrayView<FFrameNumber> OutKeyTimes) override;
-	virtual void SetKeyTimes(TArrayView<const FKeyHandle> InHandles, TArrayView<const FFrameNumber> InKeyTimes) override;
-	virtual void DuplicateKeys(TArrayView<const FKeyHandle> InHandles, TArrayView<FKeyHandle> OutNewHandles) override;
-	virtual void DeleteKeys(TArrayView<const FKeyHandle> InHandles) override;
-	virtual void DeleteKeysFrom(FFrameNumber InTime, bool bDeleteKeysBefore) override;
-	virtual void ChangeFrameResolution(FFrameRate SourceRate, FFrameRate DestinationRate) override;
-	virtual TRange<FFrameNumber> ComputeEffectiveRange() const override;
-	virtual int32 GetNumKeys() const override;
-	virtual void Reset() override;
-	virtual void Offset(FFrameNumber DeltaPosition) override;
-	virtual void Optimize(const FKeyDataOptimizationParams& InParameters) override;
-	virtual void ClearDefault() override;
+	MOVIESCENE_API virtual void GetKeys(const TRange<FFrameNumber>& WithinRange, TArray<FFrameNumber>* OutKeyTimes, TArray<FKeyHandle>* OutKeyHandles) override;
+	MOVIESCENE_API virtual void GetKeyTimes(TArrayView<const FKeyHandle> InHandles, TArrayView<FFrameNumber> OutKeyTimes) override;
+	MOVIESCENE_API virtual void SetKeyTimes(TArrayView<const FKeyHandle> InHandles, TArrayView<const FFrameNumber> InKeyTimes) override;
+	MOVIESCENE_API virtual void DuplicateKeys(TArrayView<const FKeyHandle> InHandles, TArrayView<FKeyHandle> OutNewHandles) override;
+	MOVIESCENE_API virtual void DeleteKeys(TArrayView<const FKeyHandle> InHandles) override;
+	MOVIESCENE_API virtual void DeleteKeysFrom(FFrameNumber InTime, bool bDeleteKeysBefore) override;
+	MOVIESCENE_API virtual void ChangeFrameResolution(FFrameRate SourceRate, FFrameRate DestinationRate) override;
+	MOVIESCENE_API virtual TRange<FFrameNumber> ComputeEffectiveRange() const override;
+	MOVIESCENE_API virtual int32 GetNumKeys() const override;
+	MOVIESCENE_API virtual void Reset() override;
+	MOVIESCENE_API virtual void Offset(FFrameNumber DeltaPosition) override;
+	MOVIESCENE_API virtual void Optimize(const FKeyDataOptimizationParams& InParameters) override;
+	MOVIESCENE_API virtual void ClearDefault() override;
 
 public:
 
diff --git a/Engine/Source/Runtime/MovieScene/Public/Channels/MovieSceneChannel.h b/Engine/Source/Runtime/MovieScene/Public/Channels/MovieSceneChannel.h
--- a/Engine/Source/Runtime/MovieScene/Public/Channels/MovieSceneChannel.h
+++ b/Engine/Source/Runtime/MovieScene/Public/Channels/MovieSceneChannel.h
@@ -53,7 +53,7 @@ DECLARE_MULTICAST_DELEGATE_TwoParams(FMovieSceneChannelDataKeyDeletedEvent, FMov
 DECLARE_MULTICAST_DELEGATE_TwoParams(FMovieSceneChannelDataKeyMovedEvent, FMovieSceneChannel*, const TArray<FKeyMoveEventItem>& Items);
 
 USTRUCT()
-struct MOVIESCENE_API FMovieSceneChannel
+struct FMovieSceneChannel
 {
 	GENERATED_BODY()
 
@@ -66,7 +66,7 @@ struct MOVIESCENE_API FMovieSceneChannel
 	 * @param InHandle              The handle of the key to get the time for
 	 * @param OutKeyTime            Out parameter to receive the key's time
 	 */
-	void GetKeyTime(const FKeyHandle InHandle, FFrameNumber& OutKeyTime);
+	MOVIESCENE_API void GetKeyTime(const FKeyHandle InHandle, FFrameNumber& OutKeyTime);
 
 	/**
 	 * Set the time for the specified key handle
@@ -74,7 +74,7 @@ struct MOVIESCENE_API FMovieSceneChannel
 	 * @param InHandle              The handle of the key to get the time for
 	 * @param InKeyTime             The new time for the key
 	 */
-	void SetKeyTime(const FKeyHandle InHandle, const FFrameNumber InKeyTime);
+	MOVIESCENE_API void SetKeyTime(const FKeyHandle InHandle, const FFrameNumber InKeyTime);
 
 	/**
 	 * Get key information pertaining to all keys that exist within the specified range
@@ -221,4 +221,4 @@ protected:
 	/** Broadcasts a notification whenever key is moved */
 	FMovieSceneChannelDataKeyMovedEvent KeyMovedEvent;
 
-};
\ No newline at end of file
+};
diff --git a/Engine/Source/Runtime/MovieScene/Public/Channels/MovieSceneChannelData.h b/Engine/Source/Runtime/MovieScene/Public/Channels/MovieSceneChannelData.h
--- a/Engine/Source/Runtime/MovieScene/Public/Channels/MovieSceneChannelData.h
+++ b/Engine/Source/Runtime/MovieScene/Public/Channels/MovieSceneChannelData.h
@@ -96,7 +96,7 @@ namespace MovieScene
  * Base class channel data utility that provides a consistent interface to a sorted array of times and handles.
  * Complete access should be through TMovieSceneChannelData that allows mutation of the data
  */
-struct MOVIESCENE_API FMovieSceneChannelData
+struct FMovieSceneChannelData
 {
 	/**
 	 * Read-only access to this channel's key times.
@@ -121,7 +121,7 @@ struct MOVIESCENE_API FMovieSceneChannelData
 	 * @param Index          The index to retrieve
 	 * @return A key handle that identifies the key at the specified index, regardless of re-ordering
 	 */
-	FKeyHandle GetHandle(int32 Index);
+	MOVIESCENE_API FKeyHandle GetHandle(int32 Index);
 
 	/**
 	 * Attempt to retrieve the index of key from its handle
@@ -129,7 +129,7 @@ struct MOVIESCENE_API FMovieSceneChannelData
 	 * @param Handle         The handle to retrieve
 	 * @return The index of the key, or INDEX_NONE
 	 */
-	int32 GetIndex(FKeyHandle Handle);
+	MOVIESCENE_API int32 GetIndex(FKeyHandle Handle);
 
 	/**
 	 * Attempt to find a key at a given time and tolerance
@@ -138,7 +138,7 @@ struct MOVIESCENE_API FMovieSceneChannelData
 	 * @param InTolerance    A tolerance of frame numbers to allow either side of the specified time
 	 * @return The index of the key closest to InTime and within InTolerance, or INDEX_NONE
 	 */
-	int32 FindKey(FFrameNumber InTime, FFrameNumber InTolerance = 0);
+	MOVIESCENE_API int32 FindKey(FFrameNumber InTime, FFrameNumber InTolerance = 0);
 
 	/**
 	 * Find the range of keys that fall around InTime +/- InTolerance up to a maximum
@@ -149,14 +149,14 @@ struct MOVIESCENE_API FMovieSceneChannelData
 	 * @param OutMax         The latest index that met the conditions of the search
 	 * @param InTolerance    The tolerance range to search around PredicateTime with
 	 */
-	void FindKeys(FFrameNumber InTime, int32 MaxNum, int32& OutMinIndex, int32& OutMaxIndex, int32 InTolerance);
+	MOVIESCENE_API void FindKeys(FFrameNumber InTime, int32 MaxNum, int32& OutMinIndex, int32& OutMaxIndex, int32 InTolerance);
 
 	/**
 	 * Compute the total time range of the channel data.
 	 *
 	 * @return The range of this channel data
 	 */
-	TRange<FFrameNumber> GetTotalRange() const;
+	MOVIESCENE_API TRange<FFrameNumber> GetTotalRange() const;
 
 	/**
 	 * Convert the frame resolution of a movie scene channel by moving the key times to the equivalent frame time
@@ -164,7 +164,7 @@ struct MOVIESCENE_API FMovieSceneChannelData
 	 * @param SourceRate      The frame rate the channel is currently in
 	 * @param DestinationRate The new frame rate to convert the channel to
 	 */
-	void ChangeFrameResolution(FFrameRate SourceRate, FFrameRate DestinationRate);
+	MOVIESCENE_API void ChangeFrameResolution(FFrameRate SourceRate, FFrameRate DestinationRate);
 
 	/**
 	 * Get all the keys in the given range. Resulting arrays must be the same size where indices correspond to both arrays.
@@ -173,7 +173,7 @@ struct MOVIESCENE_API FMovieSceneChannelData
 	 * @param OutKeyTimes        Array to receive all key times within the given range
 	 * @param OutKeyHandles      Array to receive all key handles within the given range
 	 */
-	void GetKeys(const TRange<FFrameNumber>& WithinRange, TArray<FFrameNumber>* OutKeyTimes, TArray<FKeyHandle>* OutKeyHandles);
+	MOVIESCENE_API void GetKeys(const TRange<FFrameNumber>& WithinRange, TArray<FFrameNumber>* OutKeyTimes, TArray<FKeyHandle>* OutKeyHandles);
 
 	/**
 	 * Get key times for a number of keys in the channel data
@@ -181,14 +181,14 @@ struct MOVIESCENE_API FMovieSceneChannelData
 	 * @param InHandles          Array of key handles that should have their times set
 	 * @param OutKeyTimes        Array of times that should be set for each key handle. Must be exactly the size of InHandles
 	 */
-	void GetKeyTimes(TArrayView<const FKeyHandle> InHandles, TArrayView<FFrameNumber> OutKeyTimes);
+	MOVIESCENE_API void GetKeyTimes(TArrayView<const FKeyHandle> InHandles, TArrayView<FFrameNumber> OutKeyTimes);
 
 	/**
 	 * Offset the channel data by a given delta time
 	 *
 	 * @param DeltaTime     The time to offset by
 	 */
-	void Offset(FFrameNumber DeltaTime);
+	MOVIESCENE_API void Offset(FFrameNumber DeltaTime);
 
 protected:
 
@@ -199,21 +199,21 @@ protected:
 	 * @param InKeyHandles   A key handle map used for persistent, order independent identification of keys
 	 * @param InChannel      A optional pointer to the owning channel.
 	 */
-	FMovieSceneChannelData(TArray<FFrameNumber>* InTimes, FKeyHandleLookupTable* InKeyHandles, FMovieSceneChannel* InChannel = nullptr);
+	MOVIESCENE_API FMovieSceneChannelData(TArray<FFrameNumber>* InTimes, FKeyHandleLookupTable* InKeyHandles, FMovieSceneChannel* InChannel = nullptr);
 
 	/**
 	 * Move the key at index KeyIndex to a new time
 	 *
 	 * @return The index of the key in its new position
 	 */
-	int32 MoveKeyInternal(int32 KeyIndex, FFrameNumber InNewTime);
+	MOVIESCENE_API int32 MoveKeyInternal(int32 KeyIndex, FFrameNumber InNewTime);
 
 	/**
 	 * Add a new key at the specified time
 	 *
 	 * @return The index of the key in its new position
 	 */
-	int32 AddKeyInternal(FFrameNumber InTime);
+	MOVIESCENE_API int32 AddKeyInternal(FFrameNumber InTime);
 
 
 protected:
@@ -605,4 +605,4 @@ private:
 
 	/** Pointer to an external array of values, to be kept in sync with Times. */
 	const TArray<ValueType>* Values;
-};
\ No newline at end of file
+};
diff --git a/Engine/Source/Runtime/MovieScene/Public/Channels/MovieSceneChannelEditorData.h b/Engine/Source/Runtime/MovieScene/Public/Channels/MovieSceneChannelEditorData.h
--- a/Engine/Source/Runtime/MovieScene/Public/Channels/MovieSceneChannelEditorData.h
+++ b/Engine/Source/Runtime/MovieScene/Public/Channels/MovieSceneChannelEditorData.h
@@ -127,21 +127,21 @@ struct TMovieSceneExternalValue
 /**
  * Commonly used channel display names and colors
  */
-struct MOVIESCENE_API FCommonChannelData
+struct FCommonChannelData
 {
-	static const FText ChannelX;
-	static const FText ChannelY;
-	static const FText ChannelZ;
-	static const FText ChannelW;
-
-	static const FText ChannelR;
-	static const FText ChannelG;
-	static const FText ChannelB;
-	static const FText ChannelA;
-
-	static const FLinearColor RedChannelColor;
-	static const FLinearColor GreenChannelColor;
-	static const FLinearColor BlueChannelColor;
+	static MOVIESCENE_API const FText ChannelX;
+	static MOVIESCENE_API const FText ChannelY;
+	static MOVIESCENE_API const FText ChannelZ;
+	static MOVIESCENE_API const FText ChannelW;
+
+	static MOVIESCENE_API const FText ChannelR;
+	static MOVIESCENE_API const FText ChannelG;
+	static MOVIESCENE_API const FText ChannelB;
+	static MOVIESCENE_API const FText ChannelA;
+
+	static MOVIESCENE_API const FLinearColor RedChannelColor;
+	static MOVIESCENE_API const FLinearColor GreenChannelColor;
+	static MOVIESCENE_API const FLinearColor BlueChannelColor;
 };
 
 
diff --git a/Engine/Source/Runtime/MovieScene/Public/Channels/MovieSceneChannelHandle.h b/Engine/Source/Runtime/MovieScene/Public/Channels/MovieSceneChannelHandle.h
--- a/Engine/Source/Runtime/MovieScene/Public/Channels/MovieSceneChannelHandle.h
+++ b/Engine/Source/Runtime/MovieScene/Public/Channels/MovieSceneChannelHandle.h
@@ -16,17 +16,17 @@ template<typename> struct TMovieSceneChannelTraits;
 /**
  * Handle to a specific channel in a UMovieSceneSection. Will become nullptr when the FMovieSceneChannelProxy it was created with is reallocated.
  */
-struct MOVIESCENE_API FMovieSceneChannelHandle
+struct FMovieSceneChannelHandle
 {
 	/**
 	 * Default constructor
 	 */
-	FMovieSceneChannelHandle();
+	MOVIESCENE_API FMovieSceneChannelHandle();
 
 	/**
 	 * Construction from a weak channel proxy, the channel's type, and its index
 	 */
-	FMovieSceneChannelHandle(TWeakPtr<FMovieSceneChannelProxy> InWeakChannelProxy, FName InChannelTypeName, int32 InChannelIndex);
+	MOVIESCENE_API FMovieSceneChannelHandle(TWeakPtr<FMovieSceneChannelProxy> InWeakChannelProxy, FName InChannelTypeName, int32 InChannelIndex);
 
 public:
 
@@ -65,19 +65,19 @@ public:
 	/**
 	 * Access this channel's type identifier
 	 */
-	FName GetChannelTypeName() const;
+	MOVIESCENE_API FName GetChannelTypeName() const;
 
 	/**
 	 * Access this channel's index
 	 */
-	int32 GetChannelIndex() const;
+	MOVIESCENE_API int32 GetChannelIndex() const;
 
 	/**
 	 * Get the channel pointer this handle represents.
 	 *
 	 * @return the channel's pointer, or nullptr if the proxy it was created with is no longer alive.
 	 */
-	FMovieSceneChannel* Get() const;
+	MOVIESCENE_API FMovieSceneChannel* Get() const;
 
 	/**
 	 * Attempt to access the proxy for this channel. Will return nullptr if it is no longer valid.
@@ -94,14 +94,14 @@ public:
 	 *
 	 * @return the channel's meta data, or nullptr if the proxy it was created with is no longer alive.
 	 */
-	const FMovieSceneChannelMetaData* GetMetaData() const;
+	MOVIESCENE_API const FMovieSceneChannelMetaData* GetMetaData() const;
 
 	/**
 	 * Get the extended editor data associated with this channel
 	 *
 	 * @return the channel's extended editor data, or nullptr if the proxy it was created with is no longer alive.
 	 */
-	const void* GetExtendedEditorData() const;
+	MOVIESCENE_API const void* GetExtendedEditorData() const;
 
 #endif // WITH_EDITOR
 
diff --git a/Engine/Source/Runtime/MovieScene/Public/Channels/MovieSceneChannelOverrideContainer.h b/Engine/Source/Runtime/MovieScene/Public/Channels/MovieSceneChannelOverrideContainer.h
--- a/Engine/Source/Runtime/MovieScene/Public/Channels/MovieSceneChannelOverrideContainer.h
+++ b/Engine/Source/Runtime/MovieScene/Public/Channels/MovieSceneChannelOverrideContainer.h
@@ -41,8 +41,8 @@ struct FChannelOverrideEntityImportParams
 /**
  * A wrapper to implement polymorphism for FMovieSceneChannel.
  */
-UCLASS(Abstract)
-class MOVIESCENE_API UMovieSceneChannelOverrideContainer : public UMovieSceneSignedObject
+UCLASS(Abstract, MinimalAPI)
+class UMovieSceneChannelOverrideContainer : public UMovieSceneSignedObject
 {
 	GENERATED_BODY()
 
@@ -79,6 +79,6 @@ public:
 	using FOverrideCandidates = TArray<TSubclassOf<UMovieSceneChannelOverrideContainer>, TInlineAllocator<8>>;
 
 	/** Get a list of channel overrides that can work in the place of the given channel type */
-	static void GetOverrideCandidates(FName InDefaultChannelTypeName, FOverrideCandidates& OutCandidates);
+	static MOVIESCENE_API void GetOverrideCandidates(FName InDefaultChannelTypeName, FOverrideCandidates& OutCandidates);
 };
 
diff --git a/Engine/Source/Runtime/MovieScene/Public/Channels/MovieSceneChannelProxy.h b/Engine/Source/Runtime/MovieScene/Public/Channels/MovieSceneChannelProxy.h
--- a/Engine/Source/Runtime/MovieScene/Public/Channels/MovieSceneChannelProxy.h
+++ b/Engine/Source/Runtime/MovieScene/Public/Channels/MovieSceneChannelProxy.h
@@ -194,7 +194,7 @@ private:
  * Proxy type stored inside UMovieSceneSection for access to all its channels. Construction via either a single channel, or a FMovieSceneChannelProxyData structure
  * This proxy exists as a generic accessor to any channel data existing in derived types
  */
-struct MOVIESCENE_API FMovieSceneChannelProxy : TSharedFromThis<FMovieSceneChannelProxy>
+struct FMovieSceneChannelProxy : TSharedFromThis<FMovieSceneChannelProxy>
 {
 public:
 
@@ -239,7 +239,7 @@ public:
 	 *
 	 * @return A pointer to the channel, or nullptr
 	 */
-	const FMovieSceneChannelEntry* FindEntry(FName ChannelTypeName) const;
+	MOVIESCENE_API const FMovieSceneChannelEntry* FindEntry(FName ChannelTypeName) const;
 
 	/**
 	 * Find the index of the specified channel ptr in this proxy
@@ -248,7 +248,7 @@ public:
 	 * @param ChannelPtr       The channel pointer to find
 	 * @return The index of the channel if found, else INDEX_NONE
 	 */
-	int32 FindIndex(FName ChannelTypeName, const FMovieSceneChannel* ChannelPtr) const;
+	MOVIESCENE_API int32 FindIndex(FName ChannelTypeName, const FMovieSceneChannel* ChannelPtr) const;
 
 	/**
 	 * Get all channels of the specified type
@@ -271,20 +271,20 @@ public:
 	 *
 	 * @return A pointer to the channel, or nullptr if the index was invalid, or the type was not present
 	 */
-	FMovieSceneChannel* GetChannel(FName ChannelTypeName, int32 ChannelIndex) const;
+	MOVIESCENE_API FMovieSceneChannel* GetChannel(FName ChannelTypeName, int32 ChannelIndex) const;
 
 	/**
 	 * Returns the total number of channels
 	 * @return The total number of channels
 	 */
-	int32 NumChannels() const;
+	MOVIESCENE_API int32 NumChannels() const;
 
 	/**
 	 * Make a channel handle out for the specified index and channel type name
 	 *
 	 * @return A handle to the supplied channel that will become nullptr when the proxy is reallocated, or nullptr if the index or channel type name are invalid.
 	 */
-	FMovieSceneChannelHandle MakeHandle(FName ChannelTypeName, int32 Index);
+	MOVIESCENE_API FMovieSceneChannelHandle MakeHandle(FName ChannelTypeName, int32 Index);
 
 	/**
 	 * Make a channel handle out for the specified index and templated channel type
@@ -345,7 +345,7 @@ public:
 	 *
 	 * @return A handle to the channel, or an invalid handle if no channel of that name was found
 	 */
-	FMovieSceneChannelHandle GetChannelByName(FName ChannelName) const;
+	MOVIESCENE_API FMovieSceneChannelHandle GetChannelByName(FName ChannelName) const;
 
 	/**
 	 * Access all the extended data for the templated channel type
@@ -368,7 +368,7 @@ private:
 #if WITH_EDITOR
 
 	/** Populate the named channel table */
-	void EnsureHandlesByNamePopulated() const;
+	MOVIESCENE_API void EnsureHandlesByNamePopulated() const;
 
 	/** Lazy-created lookup table between a channel name and a channel handle */
 	mutable TMap<FName, FMovieSceneChannelHandle> HandlesByName;
diff --git a/Engine/Source/Runtime/MovieScene/Public/Channels/MovieSceneDoubleChannel.h b/Engine/Source/Runtime/MovieScene/Public/Channels/MovieSceneDoubleChannel.h
--- a/Engine/Source/Runtime/MovieScene/Public/Channels/MovieSceneDoubleChannel.h
+++ b/Engine/Source/Runtime/MovieScene/Public/Channels/MovieSceneDoubleChannel.h
@@ -113,7 +113,7 @@ struct TStructOpsTypeTraits<FMovieSceneDoubleValue>
 
 
 USTRUCT()
-struct MOVIESCENE_API FMovieSceneDoubleChannel : public FMovieSceneChannel
+struct FMovieSceneDoubleChannel : public FMovieSceneChannel
 {
 	GENERATED_BODY()
 
@@ -174,44 +174,44 @@ struct MOVIESCENE_API FMovieSceneDoubleChannel : public FMovieSceneChannel
 	* @param OutValue   A value to receive the result
 	* @return true if the channel was evaluated successfully, false otherwise
 	*/
-	bool Evaluate(FFrameTime InTime, double& OutValue) const;
+	MOVIESCENE_API bool Evaluate(FFrameTime InTime, double& OutValue) const;
 
 	/**
 	 * Type shortening version of Evaluate above.
 	 */
-	bool Evaluate(FFrameTime InTime, float& OutValue) const;
+	MOVIESCENE_API bool Evaluate(FFrameTime InTime, float& OutValue) const;
 
 	/**
 	 * Retrieve a cached interpolation from this channel for the specified time
 	 */
-	UE::MovieScene::Interpolation::FCachedInterpolation GetInterpolationForTime(FFrameTime InTime) const;
+	MOVIESCENE_API UE::MovieScene::Interpolation::FCachedInterpolation GetInterpolationForTime(FFrameTime InTime) const;
 
 	/**
 	 * Set the channel's times and values to the requested values
 	 */
-	void Set(TArray<FFrameNumber> InTimes, TArray<FMovieSceneDoubleValue> InValues);
+	MOVIESCENE_API void Set(TArray<FFrameNumber> InTimes, TArray<FMovieSceneDoubleValue> InValues);
 
 	/**
 	 * Set the channel's times and values to the requested values, but does not allocate key handles
 	 */
-	void SetKeysOnly(TArrayView<FFrameNumber> InTimes, TArrayView<FMovieSceneDoubleValue> InValues);
+	MOVIESCENE_API void SetKeysOnly(TArrayView<FFrameNumber> InTimes, TArrayView<FMovieSceneDoubleValue> InValues);
 public:
 
 	// ~ FMovieSceneChannel Interface
-	virtual void GetKeys(const TRange<FFrameNumber>& WithinRange, TArray<FFrameNumber>* OutKeyTimes, TArray<FKeyHandle>* OutKeyHandles) override;
-	virtual void GetKeyTimes(TArrayView<const FKeyHandle> InHandles, TArrayView<FFrameNumber> OutKeyTimes) override;
-	virtual void SetKeyTimes(TArrayView<const FKeyHandle> InHandles, TArrayView<const FFrameNumber> InKeyTimes) override;
-	virtual void DuplicateKeys(TArrayView<const FKeyHandle> InHandles, TArrayView<FKeyHandle> OutNewHandles) override;
-	virtual void DeleteKeys(TArrayView<const FKeyHandle> InHandles) override;
-	virtual void DeleteKeysFrom(FFrameNumber InTime, bool bDeleteKeysBefore) override;
-	virtual void ChangeFrameResolution(FFrameRate SourceRate, FFrameRate DestinationRate) override;
-	virtual TRange<FFrameNumber> ComputeEffectiveRange() const override;
-	virtual int32 GetNumKeys() const override;
-	virtual void Reset() override;
-	virtual void Offset(FFrameNumber DeltaPosition) override;
-	virtual void Optimize(const FKeyDataOptimizationParams& InParameters) override;
-	virtual void ClearDefault() override;
-	virtual void PostEditChange() override;
+	MOVIESCENE_API virtual void GetKeys(const TRange<FFrameNumber>& WithinRange, TArray<FFrameNumber>* OutKeyTimes, TArray<FKeyHandle>* OutKeyHandles) override;
+	MOVIESCENE_API virtual void GetKeyTimes(TArrayView<const FKeyHandle> InHandles, TArrayView<FFrameNumber> OutKeyTimes) override;
+	MOVIESCENE_API virtual void SetKeyTimes(TArrayView<const FKeyHandle> InHandles, TArrayView<const FFrameNumber> InKeyTimes) override;
+	MOVIESCENE_API virtual void DuplicateKeys(TArrayView<const FKeyHandle> InHandles, TArrayView<FKeyHandle> OutNewHandles) override;
+	MOVIESCENE_API virtual void DeleteKeys(TArrayView<const FKeyHandle> InHandles) override;
+	MOVIESCENE_API virtual void DeleteKeysFrom(FFrameNumber InTime, bool bDeleteKeysBefore) override;
+	MOVIESCENE_API virtual void ChangeFrameResolution(FFrameRate SourceRate, FFrameRate DestinationRate) override;
+	MOVIESCENE_API virtual TRange<FFrameNumber> ComputeEffectiveRange() const override;
+	MOVIESCENE_API virtual int32 GetNumKeys() const override;
+	MOVIESCENE_API virtual void Reset() override;
+	MOVIESCENE_API virtual void Offset(FFrameNumber DeltaPosition) override;
+	MOVIESCENE_API virtual void Optimize(const FKeyDataOptimizationParams& InParameters) override;
+	MOVIESCENE_API virtual void ClearDefault() override;
+	MOVIESCENE_API virtual void PostEditChange() override;
 
 public:
 
@@ -254,7 +254,7 @@ public:
 
 public:
 
-	bool Serialize(FArchive& Ar);
+	MOVIESCENE_API bool Serialize(FArchive& Ar);
 	friend FArchive& operator<<(FArchive& Ar, FMovieSceneDoubleChannel& Me)
 	{
 		Me.Serialize(Ar);
@@ -262,15 +262,15 @@ public:
 	}
 
 	/** Serialize this double channel from a mismatching property tag (FRichCurve or FMovieSceneFloatChannel) */
-	bool SerializeFromMismatchedTag(const FPropertyTag& Tag, FStructuredArchive::FSlot Slot);
+	MOVIESCENE_API bool SerializeFromMismatchedTag(const FPropertyTag& Tag, FStructuredArchive::FSlot Slot);
 
-	int32 AddConstantKey(FFrameNumber InTime, double InValue);
+	MOVIESCENE_API int32 AddConstantKey(FFrameNumber InTime, double InValue);
 
-	int32 AddLinearKey(FFrameNumber InTime, double InValue);
+	MOVIESCENE_API int32 AddLinearKey(FFrameNumber InTime, double InValue);
 
-	int32 AddCubicKey(FFrameNumber InTime, double InValue, ERichCurveTangentMode TangentMode = RCTM_Auto, const FMovieSceneTangentData& Tangent = FMovieSceneTangentData());
+	MOVIESCENE_API int32 AddCubicKey(FFrameNumber InTime, double InValue, ERichCurveTangentMode TangentMode = RCTM_Auto, const FMovieSceneTangentData& Tangent = FMovieSceneTangentData());
 
-	void AutoSetTangents(float Tension = 0.f);
+	MOVIESCENE_API void AutoSetTangents(float Tension = 0.f);
 
 	/** Get the channel's frame resolution */
 	FFrameRate GetTickResolution() const { return TickResolution; }
@@ -287,7 +287,7 @@ public:
 	 * @param TickResolution        The tick resolution with which to interpret this channel's times
 	 * @param InOutPoints           An array to populate with the evaluated points
 	 */
-	void PopulateCurvePoints(double StartTimeSeconds, double EndTimeSeconds, double TimeThreshold, double ValueThreshold, FFrameRate TickResolution, TArray<TTuple<double, double>>& InOutPoints) const;	
+	MOVIESCENE_API void PopulateCurvePoints(double StartTimeSeconds, double EndTimeSeconds, double TimeThreshold, double ValueThreshold, FFrameRate TickResolution, TArray<TTuple<double, double>>& InOutPoints) const;	
 
 	/**
 	* Add keys with these times to channel. The number of elements in both arrays much match or nothing is added.
@@ -295,13 +295,13 @@ public:
 	* @param InTimes Times to add
 	* @param InValues Values to add
 	*/
-	void AddKeys(const TArray<FFrameNumber>& InTimes, const TArray<FMovieSceneDoubleValue>& InValues);
+	MOVIESCENE_API void AddKeys(const TArray<FFrameNumber>& InTimes, const TArray<FMovieSceneDoubleValue>& InValues);
 
 #if WITH_EDITORONLY_DATA
 	/* Get whether to show this curve in the UI */
-	bool GetShowCurve() const;
+	MOVIESCENE_API bool GetShowCurve() const;
 	/* Set whether to show this curve in the UI */
-	void SetShowCurve(bool bInShowCurve);
+	MOVIESCENE_API void SetShowCurve(bool bInShowCurve);
 #endif
 
 public:
diff --git a/Engine/Source/Runtime/MovieScene/Public/Channels/MovieSceneFloatChannel.h b/Engine/Source/Runtime/MovieScene/Public/Channels/MovieSceneFloatChannel.h
--- a/Engine/Source/Runtime/MovieScene/Public/Channels/MovieSceneFloatChannel.h
+++ b/Engine/Source/Runtime/MovieScene/Public/Channels/MovieSceneFloatChannel.h
@@ -112,7 +112,7 @@ struct TStructOpsTypeTraits<FMovieSceneFloatValue>
 
 
 USTRUCT()
-struct MOVIESCENE_API FMovieSceneFloatChannel : public FMovieSceneChannel
+struct FMovieSceneFloatChannel : public FMovieSceneChannel
 {
 	GENERATED_BODY()
 
@@ -173,39 +173,39 @@ struct MOVIESCENE_API FMovieSceneFloatChannel : public FMovieSceneChannel
 	* @param OutValue   A value to receive the result
 	* @return true if the channel was evaluated successfully, false otherwise
 	*/
-	bool Evaluate(FFrameTime InTime, float& OutValue) const;
+	MOVIESCENE_API bool Evaluate(FFrameTime InTime, float& OutValue) const;
 
 	/**
 	 * Retrieve a cached interpolation from this channel for the specified time
 	 */
-	UE::MovieScene::Interpolation::FCachedInterpolation GetInterpolationForTime(FFrameTime InTime) const;
+	MOVIESCENE_API UE::MovieScene::Interpolation::FCachedInterpolation GetInterpolationForTime(FFrameTime InTime) const;
 
 	/**
 	 * Set the channel's times and values to the requested values
 	 */
-	void Set(TArray<FFrameNumber> InTimes, TArray<FMovieSceneFloatValue> InValues);
+	MOVIESCENE_API void Set(TArray<FFrameNumber> InTimes, TArray<FMovieSceneFloatValue> InValues);
 
 	/**
 	 * Set the channel's times and values to the requested values, but does not allocate key handles
 	 */
-	void SetKeysOnly(TArrayView<FFrameNumber> InTimes, TArrayView<FMovieSceneFloatValue> InValues);
+	MOVIESCENE_API void SetKeysOnly(TArrayView<FFrameNumber> InTimes, TArrayView<FMovieSceneFloatValue> InValues);
 public:
 
 	// ~ FMovieSceneChannel Interface
-	virtual void GetKeys(const TRange<FFrameNumber>& WithinRange, TArray<FFrameNumber>* OutKeyTimes, TArray<FKeyHandle>* OutKeyHandles) override;
-	virtual void GetKeyTimes(TArrayView<const FKeyHandle> InHandles, TArrayView<FFrameNumber> OutKeyTimes) override;
-	virtual void SetKeyTimes(TArrayView<const FKeyHandle> InHandles, TArrayView<const FFrameNumber> InKeyTimes) override;
-	virtual void DuplicateKeys(TArrayView<const FKeyHandle> InHandles, TArrayView<FKeyHandle> OutNewHandles) override;
-	virtual void DeleteKeys(TArrayView<const FKeyHandle> InHandles) override;
-	virtual void DeleteKeysFrom(FFrameNumber InTime, bool bDeleteKeysBefore) override;
-	virtual void ChangeFrameResolution(FFrameRate SourceRate, FFrameRate DestinationRate) override;
-	virtual TRange<FFrameNumber> ComputeEffectiveRange() const override;
-	virtual int32 GetNumKeys() const override;
-	virtual void Reset() override;
-	virtual void Offset(FFrameNumber DeltaPosition) override;
-	virtual void Optimize(const FKeyDataOptimizationParams& InParameters) override;
-	virtual void ClearDefault() override;
-	virtual void PostEditChange() override;
+	MOVIESCENE_API virtual void GetKeys(const TRange<FFrameNumber>& WithinRange, TArray<FFrameNumber>* OutKeyTimes, TArray<FKeyHandle>* OutKeyHandles) override;
+	MOVIESCENE_API virtual void GetKeyTimes(TArrayView<const FKeyHandle> InHandles, TArrayView<FFrameNumber> OutKeyTimes) override;
+	MOVIESCENE_API virtual void SetKeyTimes(TArrayView<const FKeyHandle> InHandles, TArrayView<const FFrameNumber> InKeyTimes) override;
+	MOVIESCENE_API virtual void DuplicateKeys(TArrayView<const FKeyHandle> InHandles, TArrayView<FKeyHandle> OutNewHandles) override;
+	MOVIESCENE_API virtual void DeleteKeys(TArrayView<const FKeyHandle> InHandles) override;
+	MOVIESCENE_API virtual void DeleteKeysFrom(FFrameNumber InTime, bool bDeleteKeysBefore) override;
+	MOVIESCENE_API virtual void ChangeFrameResolution(FFrameRate SourceRate, FFrameRate DestinationRate) override;
+	MOVIESCENE_API virtual TRange<FFrameNumber> ComputeEffectiveRange() const override;
+	MOVIESCENE_API virtual int32 GetNumKeys() const override;
+	MOVIESCENE_API virtual void Reset() override;
+	MOVIESCENE_API virtual void Offset(FFrameNumber DeltaPosition) override;
+	MOVIESCENE_API virtual void Optimize(const FKeyDataOptimizationParams& InParameters) override;
+	MOVIESCENE_API virtual void ClearDefault() override;
+	MOVIESCENE_API virtual void PostEditChange() override;
 
 public:
 
@@ -248,9 +248,9 @@ public:
 
 public:
 
-	bool Serialize(FArchive& Ar);
+	MOVIESCENE_API bool Serialize(FArchive& Ar);
 #if WITH_EDITORONLY_DATA
-	void PostSerialize(const FArchive& Ar);
+	MOVIESCENE_API void PostSerialize(const FArchive& Ar);
 #endif
 	friend FArchive& operator<<(FArchive& Ar, FMovieSceneFloatChannel& Me)
 	{
@@ -259,15 +259,15 @@ public:
 	}
 
 	/** Serialize this float function from a mismatching property tag (FRichCurve) */
-	bool SerializeFromMismatchedTag(const FPropertyTag& Tag, FStructuredArchive::FSlot Slot);
+	MOVIESCENE_API bool SerializeFromMismatchedTag(const FPropertyTag& Tag, FStructuredArchive::FSlot Slot);
 
-	int32 AddConstantKey(FFrameNumber InTime, float InValue);
+	MOVIESCENE_API int32 AddConstantKey(FFrameNumber InTime, float InValue);
 
-	int32 AddLinearKey(FFrameNumber InTime, float InValue);
+	MOVIESCENE_API int32 AddLinearKey(FFrameNumber InTime, float InValue);
 
-	int32 AddCubicKey(FFrameNumber InTime, float InValue, ERichCurveTangentMode TangentMode = RCTM_Auto, const FMovieSceneTangentData& Tangent = FMovieSceneTangentData());
+	MOVIESCENE_API int32 AddCubicKey(FFrameNumber InTime, float InValue, ERichCurveTangentMode TangentMode = RCTM_Auto, const FMovieSceneTangentData& Tangent = FMovieSceneTangentData());
 
-	void AutoSetTangents(float Tension = 0.f);
+	MOVIESCENE_API void AutoSetTangents(float Tension = 0.f);
 
 	/** Get the channel's frame resolution */
 	FFrameRate GetTickResolution() const { return TickResolution; }
@@ -284,7 +284,7 @@ public:
 	 * @param TickResolution        The tick resolution with which to interpret this channel's times
 	 * @param InOutPoints           An array to populate with the evaluated points
 	 */
-	void PopulateCurvePoints(double StartTimeSeconds, double EndTimeSeconds, double TimeThreshold, float ValueThreshold, FFrameRate TickResolution, TArray<TTuple<double, double>>& InOutPoints) const;	
+	MOVIESCENE_API void PopulateCurvePoints(double StartTimeSeconds, double EndTimeSeconds, double TimeThreshold, float ValueThreshold, FFrameRate TickResolution, TArray<TTuple<double, double>>& InOutPoints) const;	
 
 	/**
 	* Add keys with these times to channel. The number of elements in both arrays much match or nothing is added.
@@ -292,13 +292,13 @@ public:
 	* @param InTimes Times to add
 	* @param InValues Values to add
 	*/
-	void AddKeys(const TArray<FFrameNumber>& InTimes, const TArray<FMovieSceneFloatValue>& InValues);
+	MOVIESCENE_API void AddKeys(const TArray<FFrameNumber>& InTimes, const TArray<FMovieSceneFloatValue>& InValues);
 
 #if WITH_EDITORONLY_DATA
 	/* Get whether to show this curve in the UI */
-	bool GetShowCurve() const;
+	MOVIESCENE_API bool GetShowCurve() const;
 	/* Set whether to show this curve in the UI */
-	void SetShowCurve(bool bInShowCurve);
+	MOVIESCENE_API void SetShowCurve(bool bInShowCurve);
 #endif
 
 public:
