Minor optimization in FMeshElementIndexer::BuildIndex: use inline array to avoid reallocations.

diff --git a/Engine/Source/Runtime/MeshDescription/Private/MeshElementIndexer.cpp b/Engine/Source/Runtime/MeshDescription/Private/MeshElementIndexer.cpp
--- a/Engine/Source/Runtime/MeshDescription/Private/MeshElementIndexer.cpp
+++ b/Engine/Source/Runtime/MeshDescription/Private/MeshElementIndexer.cpp
@@ -353,8 +353,7 @@ void FMeshElementIndexer::BuildIndex(int32 Index)
 	FIndexPerChannel& PerChannelIndex = PerChannelIndices[Index];
 
 	// A list of all the keys which will be updated
-	TArray<int32> KeysToUpdate;
-	KeysToUpdate.Reserve(PerChannelIndex.StaleReferencerIndices.Num());
+	TArray<int32, TInlineAllocator<256>> KeysToUpdate;
 
 	// Ensure we have enough chunks to accommodate the current number of keys.
 	// Any new ones will be default initialized to appropriate values.
