MeshDescription: Back-comp fix for meshes that somehow got saved with FTransform attributes prior to bulk-data saving being disabled for those attributes.

#rnx

diff --git a/Engine/Source/Runtime/MeshDescription/Public/MeshAttributeArray.h b/Engine/Source/Runtime/MeshDescription/Public/MeshAttributeArray.h
--- a/Engine/Source/Runtime/MeshDescription/Public/MeshAttributeArray.h
+++ b/Engine/Source/Runtime/MeshDescription/Public/MeshAttributeArray.h
@@ -11,7 +11,6 @@
 #include "CoreMinimal.h"
 #include "Delegates/IntegerSequence.h"
 #include "HAL/PlatformCrt.h"
-#include "Math/UnrealMathSSE.h"
 #include "Math/Vector.h"
 #include "Math/Vector2D.h"
 #include "Math/Vector4.h"
@@ -22,7 +21,6 @@
 #include "Misc/EnumClassFlags.h"
 #include "Misc/TVariant.h"
 #include "Serialization/Archive.h"
-#include "Serialization/StructuredArchiveAdapters.h"
 #include "Templates/CopyQualifiersFromTo.h"
 #include "Templates/EnableIf.h"
 #include "Templates/IsArray.h"
@@ -30,7 +28,7 @@
 #include "Templates/UniquePtr.h"
 #include "Templates/UnrealTemplate.h"
 #include "Templates/UnrealTypeTraits.h"
-#include "UObject/EditorObjectVersion.h"
+#include "UObject/FortniteMainBranchObjectVersion.h"
 #include "UObject/NameTypes.h"
 #include "UObject/ReleaseObjectVersion.h"
 #include "UObject/UE5MainStreamObjectVersion.h"
@@ -232,6 +230,25 @@ inline typename TEnableIf<!TIsBulkSerializable<T>::Value, FArchive>::Type& opera
 	{
 		Ar << Array.Extent;
 	}
+	
+	// A little bit prior to skeletal meshes storing their model data as mesh description, FTransform attributes were added but marked,
+	// by default, as bulk-serializable, even though FTransform doesn't support it. Therefore, this serialization path only worked on empty 
+	// FTransform attributes. However, there are still static mesh assets in the wild that contain empty FTransform attributes, and we need 
+	// to be able to successfully load them -- hence this check.
+	if constexpr (std::is_same_v<T, FTransform>)
+	{
+		if (Ar.IsLoading())
+		{
+			// This version check works because saved UStaticMesh assets set this on their archive, which is then inherited by the
+			// mesh description bulk storage.
+			const FCustomVersion* PossiblySavedVersion = Ar.GetCustomVersions().GetVersion(FFortniteMainBranchObjectVersion::GUID);
+			if (PossiblySavedVersion && PossiblySavedVersion->Version < FFortniteMainBranchObjectVersion::MeshDescriptionForSkeletalMesh)
+			{
+				Array.Container.BulkSerialize( Ar );
+				return Ar;
+			}
+		}
+	}
 
 	// Serialize types which aren't bulk serializable, which need to be serialized element-by-element
 	Ar << Array.Container;
