Added new functionality to split MeshDescription polygons into simple triangles.

diff --git a/Engine/Source/Runtime/MeshDescription/Private/MeshDescription.cpp b/Engine/Source/Runtime/MeshDescription/Private/MeshDescription.cpp
--- a/Engine/Source/Runtime/MeshDescription/Private/MeshDescription.cpp
+++ b/Engine/Source/Runtime/MeshDescription/Private/MeshDescription.cpp
@@ -1550,6 +1550,32 @@ void FMeshDescription::RemovePolygonTriangles(const FPolygonID PolygonID)
 	}
 }
 
+
+void FMeshDescription::SplitPolygon(FPolygonID OriginalPolygonID)
+{
+	TArrayView<const FTriangleID> TriangleIDs = PolygonToTriangles.Find<FTriangleID>(OriginalPolygonID);
+	if (TriangleIDs.Num() == 1)
+	{
+		return;
+	}
+
+	const FPolygonGroupID PolygonGroupID = PolygonPolygonGroups[OriginalPolygonID];
+	TArray<FTriangleID, TInlineAllocator<64>> TriangleIDsCopy(TriangleIDs);
+
+	for (int32 TriangleIndex = 1; TriangleIndex < TriangleIDsCopy.Num(); TriangleIndex++)
+	{
+		FPolygonID NewPolygonID = PolygonElements->Get().Add();
+		PolygonPolygonGroups[NewPolygonID] = PolygonGroupID;
+		PolygonGroupToPolygons.AddReferenceToKey(PolygonGroupID, NewPolygonID);
+
+		FTriangleID TriangleID = TriangleIDsCopy[TriangleIndex];
+		PolygonToTriangles.RemoveReferenceFromKey(OriginalPolygonID, TriangleID);
+		TrianglePolygons[TriangleID] = NewPolygonID;
+		PolygonToTriangles.AddReferenceToKey(NewPolygonID, TriangleID);
+	}
+}
+
+
 void FMeshDescription::CreatePolygonTriangles(const FPolygonID PolygonID, TArrayView<const FVertexInstanceID> VertexInstanceIDs)
 {
 	FPolygonGroupID PolygonGroupID = PolygonPolygonGroups[PolygonID];
diff --git a/Engine/Source/Runtime/MeshDescription/Public/MeshDescription.h b/Engine/Source/Runtime/MeshDescription/Public/MeshDescription.h
--- a/Engine/Source/Runtime/MeshDescription/Public/MeshDescription.h
+++ b/Engine/Source/Runtime/MeshDescription/Public/MeshDescription.h
@@ -1158,6 +1158,8 @@ public:
 	/** Generates triangles and internal edges for the given polygon */
 	MESHDESCRIPTION_API void ComputePolygonTriangulation(const FPolygonID PolygonID);
 
+	MESHDESCRIPTION_API void SplitPolygon(FPolygonID PolygonID);
+
 
 	//////////////////////////////////////////////////////////////////////
 	// Polygon group operations
