Replace some usages of FORCEINLINE with inline in Mesh modules.

For context about this change, see this doc: https://docs.google.com/document/d/1LjpT7NpzTecfXdFNyv5KM6sCu67827BRljNi1iaWVSE/edit?tab=t.0#heading=h.fvtz8gjukz7x

Reviewers were automatically determined. Please comment if anyone else should be added.

#rnx

diff --git a/Engine/Plugins/MeshPainting/Source/MeshPaintingToolset/Public/TMeshPaintOctree.h b/Engine/Plugins/MeshPainting/Source/MeshPaintingToolset/Public/TMeshPaintOctree.h
--- a/Engine/Plugins/MeshPainting/Source/MeshPaintingToolset/Public/TMeshPaintOctree.h
+++ b/Engine/Plugins/MeshPainting/Source/MeshPaintingToolset/Public/TMeshPaintOctree.h
@@ -31,7 +31,7 @@ struct FMeshPaintTriangleOctreeSemantics
 	*
 	* @return	Bounding box of the provided octree element
 	*/
-	FORCEINLINE static FBoxCenterAndExtent GetBoundingBox(const FMeshPaintTriangle& Element)
+	inline static FBoxCenterAndExtent GetBoundingBox(const FMeshPaintTriangle& Element)
 	{
 		return Element.BoxCenterAndExtent;
 	}
@@ -45,13 +45,13 @@ struct FMeshPaintTriangleOctreeSemantics
 	*
 	* @return	true if both octree elements are equal, false if they are not
 	*/
-	FORCEINLINE static bool AreElementsEqual(const FMeshPaintTriangle& A, const FMeshPaintTriangle& B)
+	inline static bool AreElementsEqual(const FMeshPaintTriangle& A, const FMeshPaintTriangle& B)
 	{
 		return (A.Index == B.Index);
 	}
 
 	/** Ignored for this implementation */
-	FORCEINLINE static void SetElementId(const FMeshPaintTriangle& Element, FOctreeElementId2 Id)
+	inline static void SetElementId(const FMeshPaintTriangle& Element, FOctreeElementId2 Id)
 	{
 
 	}
diff --git a/Engine/Plugins/Runtime/MeshModelingToolset/Source/ModelingComponents/Public/Scene/SceneGeometrySpatialCache.h b/Engine/Plugins/Runtime/MeshModelingToolset/Source/ModelingComponents/Public/Scene/SceneGeometrySpatialCache.h
--- a/Engine/Plugins/Runtime/MeshModelingToolset/Source/ModelingComponents/Public/Scene/SceneGeometrySpatialCache.h
+++ b/Engine/Plugins/Runtime/MeshModelingToolset/Source/ModelingComponents/Public/Scene/SceneGeometrySpatialCache.h
@@ -100,7 +100,7 @@ struct FSceneGeometryID
 	FSceneGeometryID(int32 ID) { UniqueID = ID; }
 	bool operator==(const FSceneGeometryID& Other) const { return UniqueID == Other.UniqueID; }
 };
-FORCEINLINE uint32 GetTypeHash(const FSceneGeometryID& Identifier)
+inline uint32 GetTypeHash(const FSceneGeometryID& Identifier)
 {
 	return ::GetTypeHash(Identifier.UniqueID);
 }
diff --git a/Engine/Plugins/Runtime/MeshModelingToolset/Source/ModelingComponents/Public/Scene/WorldRenderCapture.h b/Engine/Plugins/Runtime/MeshModelingToolset/Source/ModelingComponents/Public/Scene/WorldRenderCapture.h
--- a/Engine/Plugins/Runtime/MeshModelingToolset/Source/ModelingComponents/Public/Scene/WorldRenderCapture.h
+++ b/Engine/Plugins/Runtime/MeshModelingToolset/Source/ModelingComponents/Public/Scene/WorldRenderCapture.h
@@ -130,7 +130,7 @@ struct TRenderCaptureTypeData
 	T DeviceDepth;
 
 	/** @return mutable data corresponding to the given CaptureType */
-	FORCEINLINE T& operator[](ERenderCaptureType CaptureType)
+	inline T& operator[](ERenderCaptureType CaptureType)
 	{
 		switch (CaptureType)
 		{
@@ -161,7 +161,7 @@ struct TRenderCaptureTypeData
 	}
 
 	/** @return constant data corresponding to the given CaptureType */
-	FORCEINLINE const T& operator[](ERenderCaptureType CaptureType) const
+	inline const T& operator[](ERenderCaptureType CaptureType) const
 	{
 		return const_cast<TRenderCaptureTypeData<T>*>(this)->operator[](CaptureType);
 	}
diff --git a/Engine/Plugins/Runtime/MeshModelingToolset/Source/ModelingOperators/Public/BackgroundModelingComputeSource.h b/Engine/Plugins/Runtime/MeshModelingToolset/Source/ModelingOperators/Public/BackgroundModelingComputeSource.h
--- a/Engine/Plugins/Runtime/MeshModelingToolset/Source/ModelingOperators/Public/BackgroundModelingComputeSource.h
+++ b/Engine/Plugins/Runtime/MeshModelingToolset/Source/ModelingOperators/Public/BackgroundModelingComputeSource.h
@@ -57,7 +57,7 @@ protected:
 	}
 
 	// FAsyncTask framework required function
-	FORCEINLINE TStatId GetStatId() const
+	inline TStatId GetStatId() const
 	{
 		RETURN_QUICK_DECLARE_CYCLE_STAT(TModelingOpTask, STATGROUP_ThreadPoolAsyncTasks);
 	}
diff --git a/Engine/Plugins/Runtime/MeshModelingToolset/Source/ModelingOperators/Public/ModelingTaskTypes.h b/Engine/Plugins/Runtime/MeshModelingToolset/Source/ModelingOperators/Public/ModelingTaskTypes.h
--- a/Engine/Plugins/Runtime/MeshModelingToolset/Source/ModelingOperators/Public/ModelingTaskTypes.h
+++ b/Engine/Plugins/Runtime/MeshModelingToolset/Source/ModelingOperators/Public/ModelingTaskTypes.h
@@ -47,7 +47,7 @@ public:
 
 
 	// required for task system
-	FORCEINLINE TStatId GetStatId() const
+	inline TStatId GetStatId() const
 	{
 		RETURN_QUICK_DECLARE_CYCLE_STAT(TDeleterTask, STATGROUP_ThreadPoolAsyncTasks);
 	}
diff --git a/Engine/Source/Runtime/MeshConversion/Public/ToDynamicMesh.h b/Engine/Source/Runtime/MeshConversion/Public/ToDynamicMesh.h
--- a/Engine/Source/Runtime/MeshConversion/Public/ToDynamicMesh.h
+++ b/Engine/Source/Runtime/MeshConversion/Public/ToDynamicMesh.h
@@ -190,7 +190,7 @@ struct TVertexAttr
 };
 
 template <typename AttrType>
-FORCEINLINE uint32 GetTypeHash(const TVertexAttr<AttrType>& Vector)
+inline uint32 GetTypeHash(const TVertexAttr<AttrType>& Vector)
 {
 	// ugh copied from FVector clearly should not be using CRC for hash!!
 	return FCrc::MemCrc32(&Vector, sizeof(Vector));
diff --git a/Engine/Source/Runtime/MeshDescription/Public/AttributeArrayContainer.h b/Engine/Source/Runtime/MeshDescription/Public/AttributeArrayContainer.h
--- a/Engine/Source/Runtime/MeshDescription/Public/AttributeArrayContainer.h
+++ b/Engine/Source/Runtime/MeshDescription/Public/AttributeArrayContainer.h
@@ -710,8 +710,8 @@ public:
 	* DO NOT USE DIRECTLY
 	* STL-like iterators to enable range-based for loop support.
 	*/
-	FORCEINLINE AttributeType* begin() const { return GetData(); }
-	FORCEINLINE AttributeType* end  () const { return GetData() + Num(); }
+	inline AttributeType* begin() const { return GetData(); }
+	inline AttributeType* end  () const { return GetData() + Num(); }
 
 private:
 	ArrayType* Array;
diff --git a/Engine/Source/Runtime/MeshDescription/Public/MeshAttributeArray.h b/Engine/Source/Runtime/MeshDescription/Public/MeshAttributeArray.h
--- a/Engine/Source/Runtime/MeshDescription/Public/MeshAttributeArray.h
+++ b/Engine/Source/Runtime/MeshDescription/Public/MeshAttributeArray.h
@@ -91,14 +91,14 @@ public:
 	template <typename T> friend typename TEnableIf<TIsBulkSerializable<T>::Value, FArchive>::Type& operator<<( FArchive& Ar, TMeshAttributeArrayBase<T>& Array );
 
 	/** Return size of container */
-	FORCEINLINE int32 Num() const { return Container.Num() / Extent; }
+	inline int32 Num() const { return Container.Num() / Extent; }
 
 	/** Return base of data */
 	UE_DEPRECATED(5.0, "This method will be removed.")
-	FORCEINLINE const AttributeType* GetData() const { return Container.GetData(); }
+	inline const AttributeType* GetData() const { return Container.GetData(); }
 
 	/** Initializes the array to the given size with the default value */
-	FORCEINLINE void Initialize(const int32 ElementCount, const AttributeType& Default)
+	inline void Initialize(const int32 ElementCount, const AttributeType& Default)
 	{
 		Container.Reset(ElementCount * Extent);
 		if (ElementCount > 0)
@@ -143,15 +143,15 @@ public:
 
 	/** Element accessors */
 	UE_DEPRECATED(5.0, "Please use GetElementBase() instead.")
-	FORCEINLINE const AttributeType& operator[](const int32 Index) const { return Container[Index]; }
+	inline const AttributeType& operator[](const int32 Index) const { return Container[Index]; }
 
 	UE_DEPRECATED(5.0, "Please use GetElementBase() instead.")
-	FORCEINLINE AttributeType& operator[](const int32 Index) { return Container[Index]; }
+	inline AttributeType& operator[](const int32 Index) { return Container[Index]; }
 
-	FORCEINLINE const AttributeType* GetElementBase(const int32 Index) const { return &Container[Index * Extent]; }
-	FORCEINLINE AttributeType* GetElementBase(const int32 Index) { return &Container[Index * Extent]; }
+	inline const AttributeType* GetElementBase(const int32 Index) const { return &Container[Index * Extent]; }
+	inline AttributeType* GetElementBase(const int32 Index) { return &Container[Index * Extent]; }
 
-	FORCEINLINE uint32 GetExtent() const { return Extent; }
+	inline uint32 GetExtent() const { return Extent; }
 
 protected:
 	/** The actual container, represented by a regular array */
@@ -317,7 +317,7 @@ class FMeshAttributeArraySetBase
 {
 public:
 	/** Constructor */
-	FORCEINLINE FMeshAttributeArraySetBase(const uint32 InType, const EMeshAttributeFlags InFlags, const int32 InNumberOfElements, const uint32 InExtent)
+	inline FMeshAttributeArraySetBase(const uint32 InType, const EMeshAttributeFlags InFlags, const int32 InNumberOfElements, const uint32 InExtent)
 		: Type(InType),
 		  Extent(InExtent),
 		  NumElements(InNumberOfElements),
@@ -351,25 +351,25 @@ public:
 
 	/** Determine whether this attribute array set is of the given type */
 	template <typename T>
-	FORCEINLINE bool HasType() const
+	inline bool HasType() const
 	{
 		return TTupleIndex<T, AttributeTypes>::Value == Type;
 	}
 
 	/** Get the type index of this attribute array set */
-	FORCEINLINE uint32 GetType() const { return Type; }
+	inline uint32 GetType() const { return Type; }
 
 	/** Get the type extent of this attribute array set */
-	FORCEINLINE uint32 GetExtent() const { return Extent; }
+	inline uint32 GetExtent() const { return Extent; }
 
 	/** Get the flags for this attribute array set */
-	FORCEINLINE EMeshAttributeFlags GetFlags() const { return Flags; }
+	inline EMeshAttributeFlags GetFlags() const { return Flags; }
 
 	/** Set the flags for this attribute array set */
-	FORCEINLINE void SetFlags(const EMeshAttributeFlags InFlags) { Flags = InFlags; }
+	inline void SetFlags(const EMeshAttributeFlags InFlags) { Flags = InFlags; }
 
 	/** Return number of elements each attribute index has */
-	FORCEINLINE int32 GetNumElements() const { return NumElements; }
+	inline int32 GetNumElements() const { return NumElements; }
 
 protected:
 	/** Type of the attribute array (based on the tuple element index from AttributeTypes) */
@@ -398,11 +398,11 @@ class TMeshAttributeArraySet final : public FMeshAttributeArraySetBase
 
 public:
 	/** Constructors */
-	FORCEINLINE explicit TMeshAttributeArraySet(const int32 Extent = 1)
+	inline explicit TMeshAttributeArraySet(const int32 Extent = 1)
 		: Super(TTupleIndex<AttributeType, AttributeTypes>::Value, EMeshAttributeFlags::None, 0, Extent)
 	{}
 
-	FORCEINLINE explicit TMeshAttributeArraySet(const int32 NumberOfChannels, const AttributeType& InDefaultValue, const EMeshAttributeFlags InFlags, const int32 InNumberOfElements, const uint32 Extent)
+	inline explicit TMeshAttributeArraySet(const int32 NumberOfChannels, const AttributeType& InDefaultValue, const EMeshAttributeFlags InFlags, const int32 InNumberOfElements, const uint32 Extent)
 		: Super(TTupleIndex<AttributeType, AttributeTypes>::Value, InFlags, InNumberOfElements, Extent),
 		  DefaultValue(InDefaultValue)
 	{
@@ -533,16 +533,16 @@ public:
 
 
 	UE_DEPRECATED(5.0, "Please use GetArrayForChannel().")
-	FORCEINLINE const TMeshAttributeArrayBase<AttributeType>& GetArrayForIndex( const int32 Index ) const { return ArrayForChannels[ Index ]; }
+	inline const TMeshAttributeArrayBase<AttributeType>& GetArrayForIndex( const int32 Index ) const { return ArrayForChannels[ Index ]; }
 	UE_DEPRECATED(5.0, "Please use GetArrayForChannel().")
-	FORCEINLINE TMeshAttributeArrayBase<AttributeType>& GetArrayForIndex( const int32 Index ) { return ArrayForChannels[ Index ]; }
+	inline TMeshAttributeArrayBase<AttributeType>& GetArrayForIndex( const int32 Index ) { return ArrayForChannels[ Index ]; }
 
 	/** Return the TMeshAttributeArrayBase corresponding to the given attribute channel */
-	FORCEINLINE const TMeshAttributeArrayBase<AttributeType>& GetArrayForChannel( const int32 Index ) const { return ArrayForChannels[ Index ]; }
-	FORCEINLINE TMeshAttributeArrayBase<AttributeType>& GetArrayForChannel( const int32 Index ) { return ArrayForChannels[ Index ]; }
+	inline const TMeshAttributeArrayBase<AttributeType>& GetArrayForChannel( const int32 Index ) const { return ArrayForChannels[ Index ]; }
+	inline TMeshAttributeArrayBase<AttributeType>& GetArrayForChannel( const int32 Index ) { return ArrayForChannels[ Index ]; }
 
 	/** Return default value for this attribute type */
-	FORCEINLINE AttributeType GetDefaultValue() const { return DefaultValue; }
+	inline AttributeType GetDefaultValue() const { return DefaultValue; }
 
 	/** Serializer */
 	friend FArchive& operator<<(FArchive& Ar, TMeshAttributeArraySet& AttributeArraySet)
@@ -574,11 +574,11 @@ class TMeshUnboundedAttributeArraySet final : public FMeshAttributeArraySetBase
 
 public:
 	/** Constructors */
-	FORCEINLINE TMeshUnboundedAttributeArraySet()
+	inline TMeshUnboundedAttributeArraySet()
 		: Super(TTupleIndex<AttributeType, AttributeTypes>::Value, EMeshAttributeFlags::None, 0, 0)
 	{}
 
-	FORCEINLINE explicit TMeshUnboundedAttributeArraySet(const int32 NumberOfChannels, const AttributeType& InDefaultValue, const EMeshAttributeFlags InFlags, const int32 InNumberOfElements)
+	inline explicit TMeshUnboundedAttributeArraySet(const int32 NumberOfChannels, const AttributeType& InDefaultValue, const EMeshAttributeFlags InFlags, const int32 InNumberOfElements)
 		: Super(TTupleIndex<AttributeType, AttributeTypes>::Value, InFlags, InNumberOfElements, 0),
 		DefaultValue(InDefaultValue)
 	{
@@ -708,11 +708,11 @@ public:
 	}
 
 	/** Return the TMeshAttributeArrayBase corresponding to the given attribute channel */
-	FORCEINLINE const TAttributeArrayContainer<AttributeType>& GetArrayForChannel( const int32 Index ) const { return ArrayForChannels[ Index ]; }
-	FORCEINLINE TAttributeArrayContainer<AttributeType>& GetArrayForChannel( const int32 Index ) { return ArrayForChannels[ Index ]; }
+	inline const TAttributeArrayContainer<AttributeType>& GetArrayForChannel( const int32 Index ) const { return ArrayForChannels[ Index ]; }
+	inline TAttributeArrayContainer<AttributeType>& GetArrayForChannel( const int32 Index ) { return ArrayForChannels[ Index ]; }
 
 	/** Return default value for this attribute type */
-	FORCEINLINE AttributeType GetDefaultValue() const { return DefaultValue; }
+	inline AttributeType GetDefaultValue() const { return DefaultValue; }
 
 	/** Serializer */
 	friend FArchive& operator<<(FArchive& Ar, TMeshUnboundedAttributeArraySet& AttributeArraySet)
@@ -1526,16 +1526,16 @@ public:
 	FAttributesSetEntry& operator=(FAttributesSetEntry&&) = default;
 
 	/** Transparent access through the TUniquePtr */
-	FORCEINLINE const FMeshAttributeArraySetBase* Get() const { return Ptr.Get(); }
-	FORCEINLINE const FMeshAttributeArraySetBase* operator->() const { return Ptr.Get(); }
-	FORCEINLINE const FMeshAttributeArraySetBase& operator*() const { return *Ptr; }
-	FORCEINLINE FMeshAttributeArraySetBase* Get() { return Ptr.Get(); }
-	FORCEINLINE FMeshAttributeArraySetBase* operator->() { return Ptr.Get(); }
-	FORCEINLINE FMeshAttributeArraySetBase& operator*() { return *Ptr; }
+	inline const FMeshAttributeArraySetBase* Get() const { return Ptr.Get(); }
+	inline const FMeshAttributeArraySetBase* operator->() const { return Ptr.Get(); }
+	inline const FMeshAttributeArraySetBase& operator*() const { return *Ptr; }
+	inline FMeshAttributeArraySetBase* Get() { return Ptr.Get(); }
+	inline FMeshAttributeArraySetBase* operator->() { return Ptr.Get(); }
+	inline FMeshAttributeArraySetBase& operator*() { return *Ptr; }
 
 	/** Object can be coerced to bool to indicate if it is valid */
-	FORCEINLINE explicit operator bool() const { return Ptr.IsValid(); }
-	FORCEINLINE bool operator!() const { return !Ptr.IsValid(); }
+	inline explicit operator bool() const { return Ptr.IsValid(); }
+	inline bool operator!() const { return !Ptr.IsValid(); }
 
 	/** Given a type at runtime, allocate an attribute array of that type, owned by Ptr */
 	void CreateArrayOfType(const uint32 Type, const uint32 Extent);
@@ -2084,13 +2084,13 @@ public:
 
 
 	/** Inserts a default-initialized value for all attributes of the given ID */
-	FORCEINLINE void Insert(const ElementIDType ElementID)
+	inline void Insert(const ElementIDType ElementID)
 	{
 		this->Insert(ElementID.GetValue());
 	}
 
 	/** Removes all attributes with the given ID */
-	FORCEINLINE void Remove(const ElementIDType ElementID)
+	inline void Remove(const ElementIDType ElementID)
 	{
 		this->Remove(ElementID.GetValue());
 	}
diff --git a/Engine/Source/Runtime/MeshDescription/Public/MeshElementArray.h b/Engine/Source/Runtime/MeshDescription/Public/MeshElementArray.h
--- a/Engine/Source/Runtime/MeshDescription/Public/MeshElementArray.h
+++ b/Engine/Source/Runtime/MeshDescription/Public/MeshElementArray.h
@@ -99,73 +99,73 @@ class TMeshElementArray final : public TMeshElementArrayBase<ElementType>
 public:
 
 	/** Resets the container, optionally reserving space for elements to be added */
-	FORCEINLINE void Reset( const int32 Elements = 0 )
+	inline void Reset( const int32 Elements = 0 )
 	{
 		Container.Reset();
 		Container.Reserve( Elements );
 	}
 
 	/** Reserves space for the specified total number of elements */
-	FORCEINLINE void Reserve( const int32 Elements ) { Container.Reserve( Elements ); }
+	inline void Reserve( const int32 Elements ) { Container.Reserve( Elements ); }
 
 	/** Add a new element at the next available index, and return the new ID */
-	FORCEINLINE ElementIDType Add() { return ElementIDType( Container.Add( ElementType() ) ); }
+	inline ElementIDType Add() { return ElementIDType( Container.Add( ElementType() ) ); }
 
 	/** Add the provided element at the next available index, and return the new ID */
-	FORCEINLINE ElementIDType Add( typename TTypeTraits<ElementType>::ConstInitType Element ) { return ElementIDType( Container.Add( Element ) ); }
+	inline ElementIDType Add( typename TTypeTraits<ElementType>::ConstInitType Element ) { return ElementIDType( Container.Add( Element ) ); }
 
 	/** Add the provided element at the next available index, and return the ID */
-	FORCEINLINE ElementIDType Add( ElementType&& Element ) { return ElementIDType( Container.Add( Forward<ElementType>( Element ) ) ); }
+	inline ElementIDType Add( ElementType&& Element ) { return ElementIDType( Container.Add( Forward<ElementType>( Element ) ) ); }
 
 	/** Inserts a new element with the given ID */
-	FORCEINLINE ElementType& Insert( const ElementIDType ID )
+	inline ElementType& Insert( const ElementIDType ID )
 	{
 		Container.Insert( ID.GetValue(), ElementType() );
 		return Container[ ID.GetValue() ];
 	}
 
 	/** Inserts the provided element with the given ID */
-	FORCEINLINE ElementType& Insert( const ElementIDType ID, typename TTypeTraits<ElementType>::ConstInitType Element )
+	inline ElementType& Insert( const ElementIDType ID, typename TTypeTraits<ElementType>::ConstInitType Element )
 	{
 		Container.Insert( ID.GetValue(), Element );
 		return Container[ ID.GetValue() ];
 	}
 
 	/** Inserts the provided element with the given ID */
-	FORCEINLINE ElementType& Insert( const ElementIDType ID, ElementType&& Element )
+	inline ElementType& Insert( const ElementIDType ID, ElementType&& Element )
 	{
 		Container.Insert( ID.GetValue(), Forward<ElementType>( Element ) );
 		return Container[ ID.GetValue() ];
 	}
 
 	/** Removes the element with the given ID */
-	FORCEINLINE void Remove( const ElementIDType ID )
+	inline void Remove( const ElementIDType ID )
 	{
 		checkSlow( Container.IsAllocated( ID.GetValue() ) );
 		Container.RemoveAt( ID.GetValue() );
 	}
 
 	/** Returns the element with the given ID */
-	FORCEINLINE ElementType& operator[]( const ElementIDType ID )
+	inline ElementType& operator[]( const ElementIDType ID )
 	{
 		checkSlow( Container.IsAllocated( ID.GetValue() ) );
 		return Container[ ID.GetValue() ];
 	}
 
-	FORCEINLINE const ElementType& operator[]( const ElementIDType ID ) const
+	inline const ElementType& operator[]( const ElementIDType ID ) const
 	{
 		checkSlow( Container.IsAllocated( ID.GetValue() ) );
 		return Container[ ID.GetValue() ];
 	}
 
 	/** Returns the number of elements in the container */
-	FORCEINLINE int32 Num() const { return Container.Num(); }
+	inline int32 Num() const { return Container.Num(); }
 
 	/** Returns the index after the last valid element */
-	FORCEINLINE int32 GetArraySize() const { return Container.GetMaxIndex(); }
+	inline int32 GetArraySize() const { return Container.GetMaxIndex(); }
 
 	/** Returns the first valid ID */
-	FORCEINLINE ElementIDType GetFirstValidID() const
+	inline ElementIDType GetFirstValidID() const
 	{
 		return Container.Num() > 0 ?
 			ElementIDType( typename TSparseArray<ElementType>::TConstIterator( Container ).GetIndex() ) :
@@ -173,13 +173,13 @@ public:
 	}
 
 	/** Returns whether the given ID is valid or not */
-	FORCEINLINE bool IsValid( const ElementIDType ID ) const
+	inline bool IsValid( const ElementIDType ID ) const
 	{
 		return ID.GetValue() >= 0 && ID.GetValue() < Container.GetMaxIndex() && Container.IsAllocated( ID.GetValue() );
 	}
 
 	/** Serializer */
-	FORCEINLINE friend FArchive& operator<<( FArchive& Ar, TMeshElementArray& Array )
+	inline friend FArchive& operator<<( FArchive& Ar, TMeshElementArray& Array )
 	{
 		Ar << static_cast<TMeshElementArrayBase<ElementType>&>( Array );
 		return Ar;
@@ -199,7 +199,7 @@ public:
 	{
 	public:
 
-		explicit FORCEINLINE TElementIDs( const TSparseArray<ElementType>& InArray )
+		explicit inline TElementIDs( const TSparseArray<ElementType>& InArray )
 			: Array( InArray )
 		{}
 
@@ -207,27 +207,27 @@ public:
 		{
 		public:
 
-			explicit FORCEINLINE TConstIterator( typename TSparseArray<ElementType>::TConstIterator&& It )
+			explicit inline TConstIterator( typename TSparseArray<ElementType>::TConstIterator&& It )
 				: Iterator( MoveTemp( It ) )
 			{}
 
-			FORCEINLINE TConstIterator& operator++()
+			inline TConstIterator& operator++()
 			{
 				++Iterator;
 				return *this;
 			}
 
-			FORCEINLINE ElementIDType operator*() const
+			inline ElementIDType operator*() const
 			{
 				return Iterator ? ElementIDType( Iterator.GetIndex() ) : INDEX_NONE;
 			}
 
-			friend FORCEINLINE bool operator==( const TConstIterator& Lhs, const TConstIterator& Rhs )
+			friend inline bool operator==( const TConstIterator& Lhs, const TConstIterator& Rhs )
 			{
 				return Lhs.Iterator == Rhs.Iterator;
 			}
 
-			friend FORCEINLINE bool operator!=( const TConstIterator& Lhs, const TConstIterator& Rhs )
+			friend inline bool operator!=( const TConstIterator& Lhs, const TConstIterator& Rhs )
 			{
 				return Lhs.Iterator != Rhs.Iterator;
 			}
@@ -237,19 +237,19 @@ public:
 			typename TSparseArray<ElementType>::TConstIterator Iterator;
 		};
 
-		FORCEINLINE TConstIterator CreateConstIterator() const
+		inline TConstIterator CreateConstIterator() const
 		{
 			return TConstIterator( typename TSparseArray<ElementType>::TConstIterator( Array ) );
 		}
 
 	public:
 
-		FORCEINLINE TConstIterator begin() const
+		inline TConstIterator begin() const
 		{
 			return TConstIterator( Array.begin() );
 		}
 
-		FORCEINLINE TConstIterator end() const
+		inline TConstIterator end() const
 		{
 			return TConstIterator( Array.end() );
 		}
@@ -260,7 +260,7 @@ public:
 	};
 
 	/** Return iterable proxy object from container */
-	TElementIDs FORCEINLINE GetElementIDs() const { return TElementIDs( Container ); }
+	TElementIDs inline GetElementIDs() const { return TElementIDs( Container ); }
 };
 
 
diff --git a/Engine/Source/Runtime/MeshDescription/Public/MeshElementContainer.h b/Engine/Source/Runtime/MeshDescription/Public/MeshElementContainer.h
--- a/Engine/Source/Runtime/MeshDescription/Public/MeshElementContainer.h
+++ b/Engine/Source/Runtime/MeshDescription/Public/MeshElementContainer.h
@@ -135,8 +135,8 @@ public:
 	}
 
 	/** Accessor for attributes */
-	FORCEINLINE FAttributesSetBase& GetAttributes() { return Attributes; }
-	FORCEINLINE const FAttributesSetBase& GetAttributes() const { return Attributes; }
+	inline FAttributesSetBase& GetAttributes() { return Attributes; }
+	inline const FAttributesSetBase& GetAttributes() const { return Attributes; }
 
 	/** Compacts elements and returns a remapping table */
 	MESHDESCRIPTION_API void Compact(TSparseArray<int32>& OutIndexRemap);
@@ -167,34 +167,34 @@ public:
 	class FElementIDs
 	{
 	public:
-		explicit FORCEINLINE FElementIDs(const TBitArray<>& InArray)
+		explicit inline FElementIDs(const TBitArray<>& InArray)
 			: Array(InArray)
 		{}
 
 		class FConstIterator
 		{
 		public:
-			explicit FORCEINLINE FConstIterator(TConstSetBitIterator<>&& It)
+			explicit inline FConstIterator(TConstSetBitIterator<>&& It)
 				: Iterator(MoveTemp(It))
 			{}
 
-			FORCEINLINE FConstIterator& operator++()
+			inline FConstIterator& operator++()
 			{
 				++Iterator;
 				return *this;
 			}
 
-			FORCEINLINE int32 operator*() const
+			inline int32 operator*() const
 			{
 				return Iterator ? Iterator.GetIndex() : INDEX_NONE;
 			}
 
-			friend FORCEINLINE bool operator==(const FConstIterator& Lhs, const FConstIterator& Rhs)
+			friend inline bool operator==(const FConstIterator& Lhs, const FConstIterator& Rhs)
 			{
 				return Lhs.Iterator == Rhs.Iterator;
 			}
 
-			friend FORCEINLINE bool operator!=(const FConstIterator& Lhs, const FConstIterator& Rhs)
+			friend inline bool operator!=(const FConstIterator& Lhs, const FConstIterator& Rhs)
 			{
 				return Lhs.Iterator != Rhs.Iterator;
 			}
@@ -203,18 +203,18 @@ public:
 			TConstSetBitIterator<> Iterator;
 		};
 
-		FORCEINLINE FConstIterator CreateConstIterator() const
+		inline FConstIterator CreateConstIterator() const
 		{
 			return FConstIterator(TConstSetBitIterator<>(Array));
 		}
 
 	public:
-		FORCEINLINE FConstIterator begin() const
+		inline FConstIterator begin() const
 		{
 			return FConstIterator(TConstSetBitIterator<>(Array));
 		}
 
-		FORCEINLINE FConstIterator end() const
+		inline FConstIterator end() const
 		{
 			return FConstIterator(TConstSetBitIterator<>(Array, Array.Num()));
 		}
@@ -224,7 +224,7 @@ public:
 	};
 
 	/** Return iterable proxy object from container */
-	FElementIDs FORCEINLINE GetElementIDs() const { return FElementIDs(BitArray); }
+	FElementIDs inline GetElementIDs() const { return FElementIDs(BitArray); }
 
 protected:
 	/** A bit array of all indices currently in use */
@@ -276,8 +276,8 @@ public:
 	}
 
 	/** Accessor for attributes */
-	FORCEINLINE TAttributesSet<ElementIDType>& GetAttributes() { return static_cast<TAttributesSet<ElementIDType>&>(Attributes); }
-	FORCEINLINE const TAttributesSet<ElementIDType>& GetAttributes() const { return static_cast<const TAttributesSet<ElementIDType>&>(Attributes); }
+	inline TAttributesSet<ElementIDType>& GetAttributes() { return static_cast<TAttributesSet<ElementIDType>&>(Attributes); }
+	inline const TAttributesSet<ElementIDType>& GetAttributes() const { return static_cast<const TAttributesSet<ElementIDType>&>(Attributes); }
 
 	/** Serializer */
 	friend FArchive& operator<<(FArchive& Ar, TMeshElementContainer& Container)
@@ -299,34 +299,34 @@ public:
 	class TElementIDs
 	{
 	public:
-		explicit FORCEINLINE TElementIDs(const TBitArray<>& InArray)
+		explicit inline TElementIDs(const TBitArray<>& InArray)
 			: Array(InArray)
 		{}
 
 		class TConstIterator
 		{
 		public:
-			explicit FORCEINLINE TConstIterator(TConstSetBitIterator<>&& It)
+			explicit inline TConstIterator(TConstSetBitIterator<>&& It)
 				: Iterator(MoveTemp(It))
 			{}
 
-			FORCEINLINE TConstIterator& operator++()
+			inline TConstIterator& operator++()
 			{
 				++Iterator;
 				return *this;
 			}
 
-			FORCEINLINE ElementIDType operator*() const
+			inline ElementIDType operator*() const
 			{
 				return ElementIDType{Iterator ? Iterator.GetIndex() : INDEX_NONE};
 			}
 
-			friend FORCEINLINE bool operator==(const TConstIterator& Lhs, const TConstIterator& Rhs)
+			friend inline bool operator==(const TConstIterator& Lhs, const TConstIterator& Rhs)
 			{
 				return Lhs.Iterator == Rhs.Iterator;
 			}
 
-			friend FORCEINLINE bool operator!=(const TConstIterator& Lhs, const TConstIterator& Rhs)
+			friend inline bool operator!=(const TConstIterator& Lhs, const TConstIterator& Rhs)
 			{
 				return Lhs.Iterator != Rhs.Iterator;
 			}
@@ -335,18 +335,18 @@ public:
 			TConstSetBitIterator<> Iterator;
 		};
 
-		FORCEINLINE TConstIterator CreateConstIterator() const
+		inline TConstIterator CreateConstIterator() const
 		{
 			return TConstIterator(TConstSetBitIterator<>(Array));
 		}
 
 	public:
-		FORCEINLINE TConstIterator begin() const
+		inline TConstIterator begin() const
 		{
 			return TConstIterator(TConstSetBitIterator<>(Array));
 		}
 
-		FORCEINLINE TConstIterator end() const
+		inline TConstIterator end() const
 		{
 			return TConstIterator(TConstSetBitIterator<>(Array, Array.Num()));
 		}
@@ -356,7 +356,7 @@ public:
 	};
 
 	/** Return iterable proxy object from container */
-	TElementIDs FORCEINLINE GetElementIDs() const { return TElementIDs(BitArray); }
+	TElementIDs inline GetElementIDs() const { return TElementIDs(BitArray); }
 };
 
 
@@ -373,14 +373,14 @@ public:
 	}
 
 	/** Transparent access through the array */
-	FORCEINLINE const FMeshElementContainer& Get(const int32 Index = 0) const { return Channels[Index]; }
-	FORCEINLINE const FMeshElementContainer& operator[](const int32 Index) const { return Channels[Index]; }
-	FORCEINLINE const FMeshElementContainer& operator*() const { return Channels[0]; }
-	FORCEINLINE const FMeshElementContainer* operator->() const { return &Channels[0]; }
-	FORCEINLINE FMeshElementContainer& Get(const int32 Index = 0) { return Channels[Index]; }
-	FORCEINLINE FMeshElementContainer& operator[](const int32 Index) { return Channels[Index]; }
-	FORCEINLINE FMeshElementContainer& operator*() { return Channels[0]; }
-	FORCEINLINE FMeshElementContainer* operator->() { return &Channels[0]; }
+	inline const FMeshElementContainer& Get(const int32 Index = 0) const { return Channels[Index]; }
+	inline const FMeshElementContainer& operator[](const int32 Index) const { return Channels[Index]; }
+	inline const FMeshElementContainer& operator*() const { return Channels[0]; }
+	inline const FMeshElementContainer* operator->() const { return &Channels[0]; }
+	inline FMeshElementContainer& Get(const int32 Index = 0) { return Channels[Index]; }
+	inline FMeshElementContainer& operator[](const int32 Index) { return Channels[Index]; }
+	inline FMeshElementContainer& operator*() { return Channels[0]; }
+	inline FMeshElementContainer* operator->() { return &Channels[0]; }
 
 	/** Change the number of indices */
 	void SetNumChannels(const int32 NumIndices) { Channels.SetNum(NumIndices); }
diff --git a/Engine/Source/Runtime/MeshDescription/Public/MeshTypes.h b/Engine/Source/Runtime/MeshDescription/Public/MeshTypes.h
--- a/Engine/Source/Runtime/MeshDescription/Public/MeshTypes.h
+++ b/Engine/Source/Runtime/MeshDescription/Public/MeshTypes.h
@@ -32,32 +32,32 @@ struct FElementID	// @todo mesheditor script: BP doesn't have name spaces, so we
 	{
 	}
 
-	FORCEINLINE int32 GetValue() const
+	inline int32 GetValue() const
 	{
 		return IDValue;
 	}
 
-	FORCEINLINE operator int32() const
+	inline operator int32() const
 	{
 		return IDValue;
 	}
 
-	FORCEINLINE bool operator==( const FElementID& Other ) const
+	inline bool operator==( const FElementID& Other ) const
 	{
 		return IDValue == Other.IDValue;
 	}
 
-	FORCEINLINE bool operator==( const int32 Other ) const
+	inline bool operator==( const int32 Other ) const
 	{
 		return IDValue == Other;
 	}
 
-	FORCEINLINE bool operator!=( const FElementID& Other ) const
+	inline bool operator!=( const FElementID& Other ) const
 	{
 		return IDValue != Other.IDValue;
 	}
 
-	FORCEINLINE bool operator!=( const int32 Other ) const
+	inline bool operator!=( const int32 Other ) const
 	{
 		return IDValue != Other;
 	}
@@ -104,7 +104,7 @@ struct FVertexID : public FElementID
 	{
 	}
 
-	FORCEINLINE friend uint32 GetTypeHash( const FVertexID& Other )
+	inline friend uint32 GetTypeHash( const FVertexID& Other )
 	{
 		return GetTypeHash( Other.IDValue );
 	}
@@ -134,7 +134,7 @@ struct FVertexInstanceID : public FElementID
 	{
 	}
 
-	FORCEINLINE friend uint32 GetTypeHash( const FVertexInstanceID& Other )
+	inline friend uint32 GetTypeHash( const FVertexInstanceID& Other )
 	{
 		return GetTypeHash( Other.IDValue );
 	}
@@ -164,7 +164,7 @@ struct FEdgeID : public FElementID
 	{
 	}
 
-	FORCEINLINE friend uint32 GetTypeHash( const FEdgeID& Other )
+	inline friend uint32 GetTypeHash( const FEdgeID& Other )
 	{
 		return GetTypeHash( Other.IDValue );
 	}
@@ -194,7 +194,7 @@ struct FUVID : public FElementID
 	{
 	}
 
-	FORCEINLINE friend uint32 GetTypeHash( const FUVID& Other )
+	inline friend uint32 GetTypeHash( const FUVID& Other )
 	{
 		return GetTypeHash( Other.IDValue );
 	}
@@ -220,7 +220,7 @@ struct FTriangleID : public FElementID
 	{
 	}
 
-	FORCEINLINE friend uint32 GetTypeHash(const FTriangleID& Other)
+	inline friend uint32 GetTypeHash(const FTriangleID& Other)
 	{
 		return GetTypeHash(Other.IDValue);
 	}
@@ -250,7 +250,7 @@ struct FPolygonGroupID : public FElementID
 	{
 	}
 
-	FORCEINLINE friend uint32 GetTypeHash( const FPolygonGroupID& Other )
+	inline friend uint32 GetTypeHash( const FPolygonGroupID& Other )
 	{
 		return GetTypeHash( Other.IDValue );
 	}
@@ -280,7 +280,7 @@ struct FPolygonID : public FElementID
 	{
 	}
 
-	FORCEINLINE friend uint32 GetTypeHash( const FPolygonID& Other )
+	inline friend uint32 GetTypeHash( const FPolygonID& Other )
 	{
 		return GetTypeHash( Other.IDValue );
 	}
diff --git a/Engine/Source/Runtime/MeshUtilitiesCommon/Public/Allocator2D.h b/Engine/Source/Runtime/MeshUtilitiesCommon/Public/Allocator2D.h
--- a/Engine/Source/Runtime/MeshUtilitiesCommon/Public/Allocator2D.h
+++ b/Engine/Source/Runtime/MeshUtilitiesCommon/Public/Allocator2D.h
@@ -139,16 +139,16 @@ private:
 		struct FStat
 		{
 #if DEBUG_LAYOUT_STATS
-			FORCEINLINE void operator++(int) { Value++; }
-			FORCEINLINE void operator+=(uint64 InValue) { Value += InValue; }
-			FORCEINLINE uint64 GetValue() const { return Value; }
+			inline void operator++(int) { Value++; }
+			inline void operator+=(uint64 InValue) { Value += InValue; }
+			inline uint64 GetValue() const { return Value; }
 		private:
 			uint64 Value = 0;
 #else
 			// These will be optimized away when stats are not activated
-			FORCEINLINE void operator++(int) { }
-			FORCEINLINE void operator+=(uint64 InValue) { }
-			FORCEINLINE uint64 GetValue() const { return 0; }
+			inline void operator++(int) { }
+			inline void operator+=(uint64 InValue) { }
+			inline uint64 GetValue() const { return 0; }
 #endif
 		};
 
@@ -174,12 +174,12 @@ private:
 };
 
 // Returns non-zero if set
-FORCEINLINE uint64 FAllocator2D::GetBit( uint32 x, uint32 y ) const
+inline uint64 FAllocator2D::GetBit( uint32 x, uint32 y ) const
 {
 	return Bits[ (x >> 6) + y * Pitch ] & ( 1ull << ( x & 63 ) );
 }
 
-FORCEINLINE void FAllocator2D::SetBit( uint32 x, uint32 y )
+inline void FAllocator2D::SetBit( uint32 x, uint32 y )
 {
 	Bits[ (x >> 6) + y * Pitch ] |= ( 1ull << ( x & 63 ) );
 
@@ -195,7 +195,7 @@ FORCEINLINE void FAllocator2D::SetBit( uint32 x, uint32 y )
 	}
 }
 
-FORCEINLINE void FAllocator2D::ClearBit( uint32 x, uint32 y )
+inline void FAllocator2D::ClearBit( uint32 x, uint32 y )
 {
 	Bits[ (x >> 6) + y * Pitch ] &= ~( 1ull << ( x & 63 ) );
 }
diff --git a/Engine/Source/Runtime/MeshUtilitiesCommon/Public/DisjointSet.h b/Engine/Source/Runtime/MeshUtilitiesCommon/Public/DisjointSet.h
--- a/Engine/Source/Runtime/MeshUtilitiesCommon/Public/DisjointSet.h
+++ b/Engine/Source/Runtime/MeshUtilitiesCommon/Public/DisjointSet.h
@@ -24,12 +24,12 @@ private:
 	TArray< uint32 >	Parents;
 };
 
-FORCEINLINE FDisjointSet::FDisjointSet( const uint32 Size )
+inline FDisjointSet::FDisjointSet( const uint32 Size )
 {
 	Init( Size );
 }
 
-FORCEINLINE void FDisjointSet::Init( uint32 Size )
+inline void FDisjointSet::Init( uint32 Size )
 {
 	Parents.SetNumUninitialized( Size, EAllowShrinking::No);
 	for( uint32 i = 0; i < Size; i++ )
@@ -38,12 +38,12 @@ FORCEINLINE void FDisjointSet::Init( uint32 Size )
 	}
 }
 
-FORCEINLINE void FDisjointSet::Reset()
+inline void FDisjointSet::Reset()
 {
 	Parents.Reset();
 }
 
-FORCEINLINE void FDisjointSet::AddDefaulted( uint32 Num )
+inline void FDisjointSet::AddDefaulted( uint32 Num )
 {
 	uint32 Start = Parents.AddUninitialized( Num );
 
diff --git a/Engine/Source/Runtime/MeshUtilitiesCommon/Public/MeshUtilitiesCommon.h b/Engine/Source/Runtime/MeshUtilitiesCommon/Public/MeshUtilitiesCommon.h
--- a/Engine/Source/Runtime/MeshUtilitiesCommon/Public/MeshUtilitiesCommon.h
+++ b/Engine/Source/Runtime/MeshUtilitiesCommon/Public/MeshUtilitiesCommon.h
@@ -51,7 +51,7 @@ struct FIndexAndZ
 /** Sorting function for vertex Z/index pairs. */
 struct FCompareIndexAndZ
 {
-	FORCEINLINE bool operator()(FIndexAndZ const& A, FIndexAndZ const& B) const { return A.Z < B.Z; }
+	inline bool operator()(FIndexAndZ const& A, FIndexAndZ const& B) const { return A.Z < B.Z; }
 };
 
 /**
diff --git a/Engine/Source/Runtime/RawMesh/Public/RawMesh.h b/Engine/Source/Runtime/RawMesh/Public/RawMesh.h
--- a/Engine/Source/Runtime/RawMesh/Public/RawMesh.h
+++ b/Engine/Source/Runtime/RawMesh/Public/RawMesh.h
@@ -85,7 +85,7 @@ struct FRawMesh
 	RAWMESH_API bool IsValidOrFixable() const;
 
 	/** Helper for getting the position of a wedge. */
-	FORCEINLINE FVector3f GetWedgePosition(int32 WedgeIndex) const
+	inline FVector3f GetWedgePosition(int32 WedgeIndex) const
 	{
 		return VertexPositions[WedgeIndices[WedgeIndex]];
 	}
@@ -144,7 +144,7 @@ public:
 	RAWMESH_API const FByteBulkData& GetBulkData() const;
 
 	/** Returns true if no bulk data is available for this mesh. */
-	FORCEINLINE bool IsEmpty() const { return BulkData.GetBulkDataSize() == 0; }
+	inline bool IsEmpty() const { return BulkData.GetBulkDataSize() == 0; }
 
 	/** Returns true if no bulk data is available for this mesh. */
 	RAWMESH_API void Empty();
