MeshDescription: Don't try to deserialize virtualized assets if the virtualized payload cannot be found.

#rnx

diff --git a/Engine/Source/Runtime/MeshDescription/Private/MeshDescription.cpp b/Engine/Source/Runtime/MeshDescription/Private/MeshDescription.cpp
--- a/Engine/Source/Runtime/MeshDescription/Private/MeshDescription.cpp
+++ b/Engine/Source/Runtime/MeshDescription/Private/MeshDescription.cpp
@@ -2190,15 +2190,21 @@ void FMeshDescriptionBulkData::LoadMeshDescription( FMeshDescription& MeshDescri
 		FRWScopeLock ScopeLock(BulkDataLock, SLT_Write);
 #endif //WITH_EDITOR
 
-		const bool bIsPersistent = true;
+		constexpr bool bIsPersistent = true;
 		UE::Serialization::FEditorBulkDataReader Ar(BulkData, bIsPersistent);
-
-		// Propagate the version information from the package to the bulk data, so that the MeshDescription
-		// is serialized with the same versioning.
-		Ar.SetUEVer(UEVersion);
-		Ar.SetLicenseeUEVer(LicenseeUEVersion);
-		Ar.SetCustomVersions(CustomVersions);
-		Ar << MeshDescription;
+	
+		// Only attempt to load the mesh description if the payload actually exists. HasPayloadData only checks if the payload size is greater
+		// than zero, not that the payload actually got loaded (which can happen asynchronously). 
+		
+		if (Ar.IsValid())
+		{
+			// Propagate the version information from the package to the bulk data, so that the MeshDescription
+			// is serialized with the same versioning.
+			Ar.SetUEVer(UEVersion);
+			Ar.SetLicenseeUEVer(LicenseeUEVersion);
+			Ar.SetCustomVersions(CustomVersions);
+			Ar << MeshDescription;
+		}
 
 		BulkData.UnloadData();
 	}
