Skip MeshDescription serializing for reference collectors. There are no uobjects inside FMeshDescription.

diff --git a/Engine/Source/Runtime/MeshDescription/Private/MeshDescriptionBase.cpp b/Engine/Source/Runtime/MeshDescription/Private/MeshDescriptionBase.cpp
--- a/Engine/Source/Runtime/MeshDescription/Private/MeshDescriptionBase.cpp
+++ b/Engine/Source/Runtime/MeshDescription/Private/MeshDescriptionBase.cpp
@@ -14,7 +14,7 @@ void UMeshDescriptionBase::Serialize(FArchive& Ar)
 
 	Super::Serialize(Ar);
 
-	if (!Ar.IsLoading() || Ar.CustomVer(FUE5MainStreamObjectVersion::GUID) >= FUE5MainStreamObjectVersion::SerializeMeshDescriptionBase)
+	if ((!Ar.IsLoading() && !Ar.IsObjectReferenceCollector()) || Ar.CustomVer(FUE5MainStreamObjectVersion::GUID) >= FUE5MainStreamObjectVersion::SerializeMeshDescriptionBase)
 	{
 		Ar << OwnedMeshDescription;
 	}
