Fixed bad conditional when checking if UMeshDescriptionBase should avoid serializing for object reference collectors.

diff --git a/Engine/Source/Runtime/MeshDescription/Private/MeshDescriptionBase.cpp b/Engine/Source/Runtime/MeshDescription/Private/MeshDescriptionBase.cpp
--- a/Engine/Source/Runtime/MeshDescription/Private/MeshDescriptionBase.cpp
+++ b/Engine/Source/Runtime/MeshDescription/Private/MeshDescriptionBase.cpp
@@ -14,9 +14,12 @@ void UMeshDescriptionBase::Serialize(FArchive& Ar)
 
 	Super::Serialize(Ar);
 
-	if ((!Ar.IsLoading() && !Ar.IsObjectReferenceCollector()) || Ar.CustomVer(FUE5MainStreamObjectVersion::GUID) >= FUE5MainStreamObjectVersion::SerializeMeshDescriptionBase)
+	if (!Ar.IsObjectReferenceCollector())
 	{
-		Ar << OwnedMeshDescription;
+		if (!Ar.IsLoading() || Ar.CustomVer(FUE5MainStreamObjectVersion::GUID) >= FUE5MainStreamObjectVersion::SerializeMeshDescriptionBase)
+		{
+			Ar << OwnedMeshDescription;
+		}
 	}
 }
 
