[UEMHC] Improved MeshDescription assert message to help diagnose a crash from crash reporter

diff --git a/Engine/Source/Runtime/MeshDescription/Public/MeshAttributeArray.h b/Engine/Source/Runtime/MeshDescription/Public/MeshAttributeArray.h
--- a/Engine/Source/Runtime/MeshDescription/Public/MeshAttributeArray.h
+++ b/Engine/Source/Runtime/MeshDescription/Public/MeshAttributeArray.h
@@ -11,6 +11,7 @@
 #include "CoreMinimal.h"
 #include "Delegates/IntegerSequence.h"
 #include "HAL/PlatformCrt.h"
+#include "Logging/StructuredLog.h"
 #include "Math/Vector.h"
 #include "Math/Vector2D.h"
 #include "Math/Vector4.h"
@@ -1762,7 +1763,25 @@ public:
 		for (auto& MapEntry : Map)
 		{
 			MapEntry.Value->Insert(Index);
-			check(MapEntry.Value->GetNumElements() == NumElements);
+
+#if DO_CHECK
+			// Log out some more debug info when this fails, as it's failing in public builds with no known repro
+			if (MapEntry.Value->GetNumElements() != NumElements)
+			{
+				UE_LOGFMT(LogMeshDescription, Error, "NumElements mismatch in FAttributesSetBase::Insert");
+				UE_LOGFMT(LogMeshDescription, Error, "NumElements: {NumElements}", NumElements);
+				UE_LOGFMT(LogMeshDescription, Error, "Index: {Index}", Index);
+				UE_LOGFMT(LogMeshDescription, Error, "Mismatched attribute: {Attribute}\n", MapEntry.Key.ToString());
+
+				for (const TPair<FName, FAttributesSetEntry>& Pair : Map)
+				{
+					UE_LOGFMT(LogMeshDescription, Error, "Attribute {Attribute} NumElements: {NumElements}\n", Pair.Key.ToString(), Pair.Value->GetNumElements());
+				}
+
+				// Add some info to the assert message anyway, in case the logs are suppressed
+				checkf(false, TEXT("NumElements mismatch in FAttributesSetBase::Insert: %s: %i != %i"), *MapEntry.Key.ToString(), MapEntry.Value->GetNumElements(), NumElements);
+			}
+#endif
 		}
 	}
 
